ChaosFlesh : Generator for ML Deformer.

diff --git a/Engine/Plugins/Animation/MLDeformer/ChaosFleshGenerator/ChaosFleshGenerator.uplugin b/Engine/Plugins/Animation/MLDeformer/ChaosFleshGenerator/ChaosFleshGenerator.uplugin
new file mode 100644
--- /dev/null
+++ b/Engine/Plugins/Animation/MLDeformer/ChaosFleshGenerator/ChaosFleshGenerator.uplugin
@@ -0,0 +1,38 @@
+{
+	"FileVersion": 3,
+	"Version": 0,
+	"VersionName": "0.0",
+	"FriendlyName": "Chaos Flesh Generator",
+	"Description": "Chaos Flesh Data Generator for ML Deformer",
+	"Category": "Animation",
+	"CreatedBy": "Epic Games, Inc.",
+	"CreatedByURL": "https://epicgames.com",
+	"DocsURL": "",
+	"MarketplaceURL": "",
+	"SupportURL": "https://forums.unrealengine.com/",
+	"CanContainContent": true,
+	"IsBetaVersion": false,
+	"IsExperimentalVersion": true,
+	"Installed": false,
+	"Modules": [
+		{
+			"Name": "ChaosFleshGenerator",
+			"Type": "Editor",
+			"LoadingPhase": "Default"
+		}
+	],
+	"Plugins": [
+		{
+			"Name": "ChaosFlesh",
+			"Enabled": true
+		},
+		{
+			"Name": "GeometryCache",
+			"Enabled": true
+		},
+		{
+			"Name": "MLDeformerFramework",
+			"Enabled": true
+		}
+	]
+}
\ No newline at end of file
diff --git a/Engine/Plugins/Animation/MLDeformer/ChaosFleshGenerator/Source/ChaosFleshGenerator/ChaosFleshGenerator.build.cs b/Engine/Plugins/Animation/MLDeformer/ChaosFleshGenerator/Source/ChaosFleshGenerator/ChaosFleshGenerator.build.cs
new file mode 100644
--- /dev/null
+++ b/Engine/Plugins/Animation/MLDeformer/ChaosFleshGenerator/Source/ChaosFleshGenerator/ChaosFleshGenerator.build.cs
@@ -0,0 +1,63 @@
+// Copyright Epic Games, Inc. All Rights Reserved.
+
+using UnrealBuildTool;
+
+public class ChaosFleshGenerator : ModuleRules
+{
+	public ChaosFleshGenerator(ReadOnlyTargetRules Target) : base(Target)
+	{
+		PCHUsage = ModuleRules.PCHUsageMode.UseExplicitOrSharedPCHs;
+		
+		PublicIncludePaths.AddRange(
+			new string[] {
+				// ... add public include paths required here ...
+			}
+			);
+				
+		
+		PrivateIncludePaths.AddRange(
+			new string[] {
+				// ... add other private include paths required here ...
+			}
+			);
+			
+		
+		PublicDependencyModuleNames.AddRange(
+			new string[]
+			{
+				"Core",
+				"ChaosCore",
+				"Chaos",
+				// ... add other public dependencies that you statically link with here ...
+			}
+			);
+			
+		
+		PrivateDependencyModuleNames.AddRange(
+			new string[]
+			{
+				"CoreUObject",
+				"ChaosFlesh",
+				"ChaosFleshEngine",
+				"ChaosFleshNodes",
+				"Engine",
+				"GeometryCache",
+				"MLDeformerFramework",
+				"MLDeformerFrameworkEditor",
+				"PropertyEditor",
+				"RenderCore",
+				"Slate",
+				"SlateCore",
+				"UnrealEd"
+			}
+			);
+		
+		
+		DynamicallyLoadedModuleNames.AddRange(
+			new string[]
+			{
+				// ... add any modules that your module loads dynamically here ...
+			}
+			);
+	}
+}
diff --git a/Engine/Plugins/Animation/MLDeformer/ChaosFleshGenerator/Source/ChaosFleshGenerator/Private/ChaosFleshGenerator.cpp b/Engine/Plugins/Animation/MLDeformer/ChaosFleshGenerator/Source/ChaosFleshGenerator/Private/ChaosFleshGenerator.cpp
new file mode 100644
--- /dev/null
+++ b/Engine/Plugins/Animation/MLDeformer/ChaosFleshGenerator/Source/ChaosFleshGenerator/Private/ChaosFleshGenerator.cpp
@@ -0,0 +1,236 @@
+// Copyright Epic Games, Inc. All Rights Reserved.
+
+#include "ChaosFleshGenerator.h"
+#include "ChaosFleshGeneratorPrivate.h"
+#include "ChaosFleshGeneratorSimulation.h"
+#include "ChaosFleshGeneratorThreading.h"
+#include "Engine/SkeletalMesh.h"
+#include "FleshGeneratorProperties.h"
+#include "Misc/AsyncTaskNotification.h"
+
+DEFINE_LOG_CATEGORY(LogChaosFleshGenerator);
+
+#define LOCTEXT_NAMESPACE "ChaosFleshGenerator"
+
+namespace UE::Chaos::FleshGenerator
+{
+	FChaosFleshGenerator::FChaosFleshGenerator()
+	{
+		Properties = TObjectPtr<UFleshGeneratorProperties>(NewObject<UFleshGeneratorProperties>());
+	}
+
+	FChaosFleshGenerator::~FChaosFleshGenerator()
+	{
+		if (TaskResource != nullptr)
+		{
+			TaskResource->FreeSimResources_GameThread();
+		}
+	}
+
+	void FChaosFleshGenerator::Tick(float DeltaTime)
+	{
+		if (PendingAction == EFleshGeneratorActions::StartGenerate)
+		{
+			StartGenerate();
+		}
+		else if (PendingAction == EFleshGeneratorActions::TickGenerate)
+		{
+			TickGenerate();
+		}
+	}
+
+	TStatId FChaosFleshGenerator::GetStatId() const
+	{
+		RETURN_QUICK_DECLARE_CYCLE_STAT(FChaosFleshGenerator, STATGROUP_Tickables);
+	}
+
+	void FChaosFleshGenerator::StartGenerate()
+	{
+		check(PendingAction == EFleshGeneratorActions::StartGenerate);
+		if (Properties->FleshAsset == nullptr)
+		{
+			UE_LOG(LogChaosFleshGenerator, Error, TEXT("FleshAsset is null."));
+			PendingAction = EFleshGeneratorActions::NoAction;
+			return;
+		}
+		if (Properties->SkeletalMeshAsset == nullptr)
+		{
+			UE_LOG(LogChaosFleshGenerator, Error, TEXT("SkeletalMeshAsset is null."));
+			PendingAction = EFleshGeneratorActions::NoAction;
+			return;
+		}
+		if (Properties->SkeletalMeshAsset->GetSkeleton() == nullptr)
+		{
+			UE_LOG(LogChaosFleshGenerator, Error, TEXT("SkeletalMeshAssets skeleton is null."));
+			PendingAction = EFleshGeneratorActions::NoAction;
+			return;
+		}
+		if (Properties->FleshAsset->SkeletalMesh != Properties->SkeletalMeshAsset)
+		{
+			UE_LOG(LogChaosFleshGenerator, Error, TEXT("Flesh assets skeletal mesh is not the same as the Generators"));
+			PendingAction = EFleshGeneratorActions::NoAction;
+			return;
+		}
+		if (Properties->AnimationSequence == nullptr)
+		{
+			UE_LOG(LogChaosFleshGenerator, Error, TEXT("AnimationSequence is null."));
+			PendingAction = EFleshGeneratorActions::NoAction;
+			return;
+		}
+		UGeometryCache* const Cache = GetCache();
+		if (Cache == nullptr)
+		{
+			UE_LOG(LogChaosFleshGenerator, Error, TEXT("Cannot find or create geometry cache."));
+			PendingAction = EFleshGeneratorActions::NoAction;
+			return;
+		}
+		if (TaskResource != nullptr)
+		{
+			UE_LOG(LogChaosFleshGenerator, Error, TEXT("Previous generation is still running."));
+			PendingAction = EFleshGeneratorActions::NoAction;
+			return;
+		}
+
+		using UE::Chaos::FleshGenerator::Private::GetMeshImportVertexMap;
+		TOptional<TArray<int32>> OptionalMap = GetMeshImportVertexMap(*Properties->SkeletalMeshAsset, *Properties->FleshAsset);
+		if (!OptionalMap)
+		{
+			PendingAction = EFleshGeneratorActions::NoAction;
+			return;
+		}
+		TaskResource = TSharedPtr<FTaskResource>(new FTaskResource);
+
+		using Private::ParseFrames;
+		using Private::Range;
+
+		// add back debugging support
+		TaskResource->FramesToSimulate = Range(Properties->AnimationSequence->GetNumberOfSampledKeys());
+		if (Properties->FramesToSimulate.Len() > 0)
+		{
+			TaskResource->FramesToSimulate = ParseFrames(Properties->FramesToSimulate);
+		}
+
+
+		const int32 NumFrames = TaskResource->FramesToSimulate.Num();
+		if (NumFrames == 0)
+		{
+			PendingAction = EFleshGeneratorActions::NoAction;
+			return;
+		}
+		TaskResource->SimulatedPositions.SetNum( NumFrames);
+
+		if (!TaskResource->AllocateSimResources_GameThread(Properties, 1 /*NumThreads*/))
+		{
+			PendingAction = EFleshGeneratorActions::NoAction;
+			return;
+		}
+		TaskResource->Cache = Cache;
+
+		TUniquePtr<FLaunchSimsTask> Task = MakeUnique<FLaunchSimsTask>(TaskResource, Properties);
+		TaskResource->Executer = MakeUnique<FTaskResource::FExecuterType>(MoveTemp(Task));
+		TaskResource->Executer->StartBackgroundTask();
+	
+		FAsyncTaskNotificationConfig NotificationConfig;
+		NotificationConfig.TitleText = LOCTEXT("SimulateFlesh", "Simulating Flesh");
+		NotificationConfig.ProgressText = FText::FromString(TEXT("0%"));
+		NotificationConfig.bCanCancel = true;
+		NotificationConfig.bKeepOpenOnSuccess = true;
+		NotificationConfig.bKeepOpenOnFailure = true;
+		TaskResource->Notification = MakeUnique<FAsyncTaskNotification>(NotificationConfig);
+		TaskResource->StartTime = FDateTime::UtcNow();
+		TaskResource->LastUpdateTime = TaskResource->StartTime;
+	
+		const TArray<int32>& Map = OptionalMap.GetValue();
+		TaskResource->ImportedVertexNumbers = TArray<uint32>(reinterpret_cast<const uint32*>(Map.GetData()), Map.Num());
+		PendingAction = EFleshGeneratorActions::TickGenerate;
+	}
+	
+	void FChaosFleshGenerator::TickGenerate()
+	{
+		check(PendingAction == EFleshGeneratorActions::TickGenerate && TaskResource != nullptr);
+			
+		bool bFinished = false;
+		const bool bCancelled = TaskResource->Notification->GetPromptAction() == EAsyncTaskNotificationPromptAction::Cancel;
+		if (TaskResource->Executer->IsDone())
+		{
+			bFinished = true;
+		}
+		else if (bCancelled)
+		{
+			TaskResource->Cancel();
+			bFinished = true;
+		}
+			
+		if (!bFinished)
+		{
+			TaskResource->FlushRendering();
+			const FDateTime CurrentTime = FDateTime::UtcNow();
+			const double SinceLastUpdate = (CurrentTime - TaskResource->LastUpdateTime).GetTotalSeconds();
+			if (SinceLastUpdate < 0.2)
+			{
+				return;
+			}
+			
+			const int32 NumSimulatedFrames = TaskResource->NumSimulatedFrames.load();
+			const int32 NumTotalFrames = TaskResource->FramesToSimulate.Num();
+			const FText ProgressMessage = FText::FromString(FString::Printf(TEXT("Finished %d/%d, %.1f%%"), NumSimulatedFrames, NumTotalFrames, 100.0 * NumSimulatedFrames / NumTotalFrames));
+			TaskResource->Notification->SetProgressText(ProgressMessage);
+			TaskResource->LastUpdateTime = CurrentTime;
+		}
+		else
+		{
+			FreeTaskResource(bCancelled);
+			PendingAction = EFleshGeneratorActions::NoAction;
+		}
+	}
+
+	UFleshGeneratorProperties& FChaosFleshGenerator::GetProperties() const
+	{
+		return *Properties;
+	}
+	
+	void FChaosFleshGenerator::RequestAction(EFleshGeneratorActions ActionType)
+	{
+		if (PendingAction != EFleshGeneratorActions::NoAction)
+		{
+			return;
+		}
+		PendingAction = ActionType;
+	}
+	
+	
+	UGeometryCache* FChaosFleshGenerator::GetCache() const
+	{
+		return Properties->SimulatedCache;
+	}
+	
+	void FChaosFleshGenerator::FreeTaskResource(bool bCancelled)
+	{
+		TaskResource->Notification->SetProgressText(LOCTEXT("Finishing", "Finishing, please wait"));
+		TaskResource->FreeSimResources_GameThread();
+		const FDateTime CurrentTime = FDateTime::UtcNow();
+		UE_LOG(LogChaosFleshGenerator, Log, TEXT("Training finished in %f seconds"), (CurrentTime - TaskResource->StartTime).GetTotalSeconds());
+	
+		{
+			UE::Chaos::FleshGenerator::Private::FTimeScope TimeScope(TEXT("Saving"));
+	
+			using UE::Chaos::FleshGenerator::Private::SaveGeometryCache;
+			using UE::Chaos::FleshGenerator::Private::SavePackage;
+			SaveGeometryCache(*TaskResource->Cache, *Properties->SkeletalMeshAsset, TaskResource->ImportedVertexNumbers, TaskResource->SimulatedPositions);
+			SavePackage(*TaskResource->Cache);
+		}
+		if (bCancelled)
+		{
+			TaskResource->Notification->SetProgressText(LOCTEXT("Cancelled", "Cancelled"));
+			TaskResource->Notification->SetComplete(false);
+		}
+		else
+		{
+			TaskResource->Notification->SetProgressText(LOCTEXT("Finished", "Finished"));
+			TaskResource->Notification->SetComplete(true);
+		}
+		TaskResource.Reset();
+	}
+};
+
+#undef LOCTEXT_NAMESPACE
diff --git a/Engine/Plugins/Animation/MLDeformer/ChaosFleshGenerator/Source/ChaosFleshGenerator/Private/ChaosFleshGenerator.h b/Engine/Plugins/Animation/MLDeformer/ChaosFleshGenerator/Source/ChaosFleshGenerator/Private/ChaosFleshGenerator.h
new file mode 100644
--- /dev/null
+++ b/Engine/Plugins/Animation/MLDeformer/ChaosFleshGenerator/Source/ChaosFleshGenerator/Private/ChaosFleshGenerator.h
@@ -0,0 +1,51 @@
+// Copyright Epic Games, Inc. All Rights Reserved.
+
+#pragma once
+
+#include "CoreTypes.h"
+#include "UObject/ObjectPtr.h"
+#include "TickableEditorObject.h"
+#include "Logging/LogMacros.h"
+
+DECLARE_LOG_CATEGORY_EXTERN(LogChaosFleshGenerator, Log, All);
+
+class UGeometryCache;
+class UFleshGeneratorProperties;
+
+namespace UE::Chaos::FleshGenerator
+{
+	struct FTaskResource;
+
+	enum class EFleshGeneratorActions
+	{
+		NoAction,
+		StartGenerate,
+		TickGenerate
+	};
+
+	class FChaosFleshGenerator : public FTickableEditorObject
+	{
+	public:
+		FChaosFleshGenerator();
+		virtual ~FChaosFleshGenerator();
+
+		//~ Begin FTickableEditorObject Interface
+		virtual void Tick(float DeltaTime) override;
+		virtual TStatId GetStatId() const override;
+		virtual ETickableTickType GetTickableTickType() const override { return ETickableTickType::Always; }
+		//~ End FTickableEditorObject Interface
+
+		UFleshGeneratorProperties& GetProperties() const;
+		void RequestAction(EFleshGeneratorActions Action);
+	private:
+
+		void StartGenerate();
+		void TickGenerate();
+		void FreeTaskResource(bool bCancelled);
+		UGeometryCache* GetCache() const;
+
+		TObjectPtr<UFleshGeneratorProperties> Properties;
+		EFleshGeneratorActions PendingAction = EFleshGeneratorActions::NoAction;
+		TSharedPtr<FTaskResource> TaskResource;
+	};
+};
\ No newline at end of file
diff --git a/Engine/Plugins/Animation/MLDeformer/ChaosFleshGenerator/Source/ChaosFleshGenerator/Private/ChaosFleshGeneratorModule.cpp b/Engine/Plugins/Animation/MLDeformer/ChaosFleshGenerator/Source/ChaosFleshGenerator/Private/ChaosFleshGeneratorModule.cpp
new file mode 100644
--- /dev/null
+++ b/Engine/Plugins/Animation/MLDeformer/ChaosFleshGenerator/Source/ChaosFleshGenerator/Private/ChaosFleshGeneratorModule.cpp
@@ -0,0 +1,27 @@
+// Copyright Epic Games, Inc. All Rights Reserved.
+#include "ChaosFleshGeneratorToolsMenuExtender.h"
+#include "MLDeformerEditorToolkit.h"
+#include "PropertyEditorModule.h"
+#include "SFleshGeneratorWidget.h"
+
+namespace UE::Chaos::FleshGenerator
+{
+	class FChaosFleshGeneratorModule : public IModuleInterface
+	{
+	public:
+		virtual void StartupModule() override
+		{
+			UE::MLDeformer::FMLDeformerEditorToolkit::AddToolsMenuExtender(CreateToolsMenuExtender());
+			FPropertyEditorModule& PropertyModule = FModuleManager::LoadModuleChecked<FPropertyEditorModule>("PropertyEditor");
+			PropertyModule.RegisterCustomClassLayout("FleshGeneratorProperties", FOnGetDetailCustomizationInstance::CreateStatic(&FFleshGeneratorDetails::MakeInstance));
+		}
+
+		virtual void ShutdownModule() override
+		{
+			FPropertyEditorModule& PropertyModule = FModuleManager::LoadModuleChecked<FPropertyEditorModule>("PropertyEditor");
+			PropertyModule.UnregisterCustomClassLayout("FleshGeneratorProperties");
+		}
+	};
+};
+
+IMPLEMENT_MODULE(UE::Chaos::FleshGenerator::FChaosFleshGeneratorModule, ChaosFleshGenerator)
\ No newline at end of file
diff --git a/Engine/Plugins/Animation/MLDeformer/ChaosFleshGenerator/Source/ChaosFleshGenerator/Private/ChaosFleshGeneratorPrivate.cpp b/Engine/Plugins/Animation/MLDeformer/ChaosFleshGenerator/Source/ChaosFleshGenerator/Private/ChaosFleshGeneratorPrivate.cpp
new file mode 100644
--- /dev/null
+++ b/Engine/Plugins/Animation/MLDeformer/ChaosFleshGenerator/Source/ChaosFleshGenerator/Private/ChaosFleshGeneratorPrivate.cpp
@@ -0,0 +1,297 @@
+// Copyright Epic Games, Inc. All Rights Reserved.
+
+#include "ChaosFleshGeneratorPrivate.h"
+
+#include "Chaos/Vector.h"
+#include "Dataflow/ChaosFleshBindingsNodes.h"
+#include "FleshGeneratorComponent.h"
+#include "FileHelpers.h"
+#include "GeometryCollection/Facades/CollectionTetrahedralBindingsFacade.h"
+#include "GeometryCache.h"
+#include "GeometryCacheCodecV1.h"
+#include "GeometryCacheMeshData.h"
+#include "GeometryCacheTrackStreamable.h"
+#include "GeometryCacheConstantTopologyWriter.h"
+#include "Internationalization/Regex.h"
+#include "Logging/LogMacros.h"
+#include "Rendering/SkeletalMeshRenderData.h"
+#include "Rendering/SkeletalMeshModel.h"
+
+DEFINE_LOG_CATEGORY(LogChaosFleshGeneratorPrivate);
+
+#define LOCTEXT_NAMESPACE "ChaosFleshGeneratorPrivate"
+
+namespace UE::Chaos::FleshGenerator
+{
+	namespace Private
+	{
+		FTimeScope::FTimeScope(FString InName)
+			: Name(MoveTemp(InName))
+			, StartTime(FDateTime::UtcNow())
+		{
+		}
+
+		FTimeScope::~FTimeScope()
+		{
+			const FTimespan Duration = FDateTime::UtcNow() - StartTime;
+			UE_LOG(LogChaosFleshGeneratorPrivate, Log, TEXT("%s took %f secs"), *Name, Duration.GetTotalSeconds());
+		}
+
+
+		void BoundSurfacePositions(
+			const USkeletalMesh* SkeletalMesh,
+			const FFleshCollection* FleshCollection, 
+			const TManagedArray<FVector3f>* RestVertices,
+			const TManagedArray<FVector3f>* SimulatedVertices,
+			TArray<FVector3f>& Positions)
+		{
+			bool bError = false;
+
+			GeometryCollection::Facades::FTetrahedralBindings TetBindings(*FleshCollection);
+			FString MeshId = UE::TetrahedralBindingsEngineUtil::GetMeshId(SkeletalMesh, false);
+			FName MeshIdName(MeshId);
+
+			const int32 LODIndex = 0;
+			const int32 TetIndex = TetBindings.GetTetMeshIndex(MeshIdName, LODIndex);
+			if (TetIndex == INDEX_NONE)
+			{
+				UE_LOG(LogChaosFleshGeneratorPrivate, Error,
+					TEXT("CreateGeometryCache - No tet mesh index associated with mesh '%s' LOD: %d"),
+					*MeshId, LODIndex);
+				bError = true;
+			}
+			if (!TetBindings.ReadBindingsGroup(TetIndex, MeshIdName, LODIndex))
+			{
+				UE_LOG(LogChaosFleshGeneratorPrivate, Error,
+					TEXT("CreateGeometryCache - Failed to read bindings group associated with mesh '%s' LOD: %d"),
+					*MeshId, LODIndex);
+				bError = true;
+			}
+
+			TUniquePtr<GeometryCollection::Facades::FTetrahedralBindings::Evaluator> BindingsEvalPtr = TetBindings.InitEvaluator(RestVertices);
+			const GeometryCollection::Facades::FTetrahedralBindings::Evaluator& BindingsEval = *BindingsEvalPtr.Get();
+			if (!BindingsEval.IsValid())
+			{
+				UE_LOG(LogChaosFleshGeneratorPrivate, Error,
+					TEXT("CreateGeometryCache - Invalid flesh bindings for skeletal mesh asset [%s]"),
+					*SkeletalMesh->GetName());
+				bError = true;
+			}
+
+			if (!bError)
+			{
+				FVector3f Min = FLT_MAX*FVector3f::One();
+				FVector3f Max = FVector3f::Zero();
+
+				TArray<::Chaos::TVector<::Chaos::FRealSingle, 3>> CurrVertices;
+				CurrVertices.SetNum(SimulatedVertices->Num());
+				for (int32 Index = 0; Index < CurrVertices.Num(); ++Index)
+				{
+					CurrVertices[Index] = ::Chaos::TVector<::Chaos::FRealSingle, 3>((*SimulatedVertices)[Index]);
+					Min = FVector3f::Min(Min, (*SimulatedVertices)[Index]);
+					Max = FVector3f::Max(Max, (*SimulatedVertices)[Index]);
+				}
+
+				const int32 NumVertices = BindingsEval.NumVertices();
+				Positions.SetNum(NumVertices);
+				for (int32 Index = 0; Index < NumVertices; ++Index)
+				{
+					Positions[Index] = BindingsEval.GetEmbeddedPosition(Index, CurrVertices);
+				}
+			}
+		}
+
+		TArray<int32> ParseFrames(const FString& FramesString)
+		{
+			TArray<int32> Result;
+			static const FRegexPattern AllowedCharsPattern(TEXT("^[-,0-9\\s]+$"));
+		
+			if (!FRegexMatcher(AllowedCharsPattern, FramesString).FindNext())
+			{
+			    UE_LOG(LogChaosFleshGeneratorPrivate, Error, TEXT("Input contains invalid characters."));
+			    return Result;
+			}
+		
+			static const FRegexPattern SingleNumberPattern(TEXT("^\\s*(\\d+)\\s*$"));
+			static const FRegexPattern RangePattern(TEXT("^\\s*(\\d+)\\s*-\\s*(\\d+)\\s*$"));
+		
+			TArray<FString> Segments;
+		    FramesString.ParseIntoArray(Segments, TEXT(","), true);
+		    for (const FString& Segment : Segments)
+		    {
+		    	bool bSegmentValid = false;
+		
+		    	FRegexMatcher SingleNumberMatcher(SingleNumberPattern, Segment);
+		    	if (SingleNumberMatcher.FindNext())
+		    	{
+		    	    const int32 SingleNumber = FCString::Atoi(*SingleNumberMatcher.GetCaptureGroup(1));
+		    	    Result.Add(SingleNumber);
+		    	    bSegmentValid = true;
+		    	}
+		    	else
+		    	{
+		    		FRegexMatcher RangeMatcher(RangePattern, Segment);
+		    		if (RangeMatcher.FindNext())
+		    		{
+		    		    const int32 RangeStart = FCString::Atoi(*RangeMatcher.GetCaptureGroup(1));
+		    		    const int32 RangeEnd = FCString::Atoi(*RangeMatcher.GetCaptureGroup(2));
+		
+		    		    for (int32 i = RangeStart; i <= RangeEnd; ++i)
+		    		    {
+		    		        Result.Add(i);
+		    		    }
+		    		    bSegmentValid = true;
+		    		}
+		    	}
+		    	
+		    	if (!bSegmentValid)
+		    	{
+		    	    UE_LOG(LogChaosFleshGeneratorPrivate, Error, TEXT("Invalid format in segment: %s"), *Segment);
+		    	}
+		    }
+		
+			return Result;
+		}
+		
+		TArray<int32> Range(int32 End)
+		{
+			TArray<int32> Result;
+			Result.Reserve(End);
+			for (int32 Index = 0; Index < End; ++Index)
+			{
+				Result.Add(Index);
+			}
+			return Result;
+		}
+		
+		int32 GetNumVertices(const FSkeletalMeshLODRenderData& LODData)
+		{
+			int32 NumVertices = 0;
+			for(const FSkelMeshRenderSection& Section : LODData.RenderSections)
+			{
+				NumVertices += Section.NumVertices;
+			}
+			return NumVertices;
+		}
+		
+		TArrayView<TArray<FVector3f>> ShrinkToValidFrames(const TArrayView<TArray<FVector3f>>& Positions, int32 NumVertices)
+		{
+			int32 NumValidFrames = 0;
+			for (const TArray<FVector3f>& Frame : Positions)
+			{
+				if (Frame.Num() != NumVertices)
+				{
+					break;
+				}
+				++NumValidFrames;
+			}
+			return TArrayView<TArray<FVector3f>>(Positions.GetData(), NumValidFrames);
+		}
+		
+		void SaveGeometryCache(UGeometryCache& GeometryCache, const USkinnedAsset& Asset, TConstArrayView<uint32> ImportedVertexNumbers, TArrayView<TArray<FVector3f>> PositionsToMoveFrom)
+		{
+			const FSkeletalMeshRenderData* RenderData = Asset.GetResourceForRendering();
+			constexpr int32 LODIndex = 0;
+			if (!RenderData || !RenderData->LODRenderData.IsValidIndex(LODIndex))
+			{
+				return;
+			}
+			const FSkeletalMeshLODRenderData& LODData = RenderData->LODRenderData[LODIndex];
+			const int32 NumVertices = GetNumVertices(LODData);
+			PositionsToMoveFrom = ShrinkToValidFrames(PositionsToMoveFrom, NumVertices);
+		
+			using UE::GeometryCacheHelpers::FGeometryCacheConstantTopologyWriter;
+			using UE::GeometryCacheHelpers::AddTrackWriterFromSkinnedAsset;
+			UE::GeometryCacheHelpers::FGeometryCacheConstantTopologyWriter::FConfig Config; Config.FPS = 24;
+			using FTrackWriter = FGeometryCacheConstantTopologyWriter::FTrackWriter;
+			FGeometryCacheConstantTopologyWriter Writer(GeometryCache, Config);
+			const int32 Index = AddTrackWriterFromSkinnedAsset(Writer, Asset);
+			if (Index == INDEX_NONE)
+			{
+				return;
+			}
+			FTrackWriter& TrackWriter = Writer.GetTrackWriter(Index);
+			TrackWriter.ImportedVertexNumbers = ImportedVertexNumbers;
+			TrackWriter.WriteAndClose(PositionsToMoveFrom);
+		}
+		
+		void SavePackage(UObject& Object)
+		{
+			TArray<UPackage*> PackagesToSave = { Object.GetOutermost() };
+			constexpr bool bCheckDirty = false;
+			constexpr bool bPromptToSave = false;
+			FEditorFileUtils::PromptForCheckoutAndSave(PackagesToSave, bCheckDirty, bPromptToSave);
+		}
+		
+		TOptional<TArray<int32>> GetMeshImportVertexMap(const USkinnedAsset& SkinnedMeshAsset, const UFleshAsset& FleshAsset)
+		{
+			constexpr int32 LODIndex = 0;
+			const TOptional<TArray<int32>> None;
+			const FSkeletalMeshModel* const MLDModel = SkinnedMeshAsset.GetImportedModel();
+			if (!MLDModel || !MLDModel->LODModels.IsValidIndex(LODIndex))
+			{
+				return None;
+			}
+			const FSkeletalMeshLODModel& MLDLOD = MLDModel->LODModels[LODIndex];
+			const TArray<int32>& Map = MLDLOD.MeshToImportVertexMap;
+			if (Map.IsEmpty())
+			{
+				UE_LOG(LogChaosFleshGeneratorPrivate, Warning, TEXT("MeshToImportVertexMap is empty. MLDeformer Asset should be an imported SkeletalMesh (e.g. from fbx)."));
+				return None;
+			}
+
+			//
+			// @todo(flesh LOD) : Add support for managing vertex mappings between skeletal LOD.
+			//		The cloth asset will extract the LOD from the ManagedArrayCollection.
+			//		The FleshGenerator will need to do the same when Flesh supports LODS. 
+			//
+
+			TArray<FVector3f> Positions;
+
+			const USkeletalMesh* SkeletalMeshAsset = Cast<USkeletalMesh>(&SkinnedMeshAsset);
+			const FFleshCollection* FleshCollection = FleshAsset.GetCollection();
+			const TManagedArray<FVector3f>* RestVertices = FleshAsset.FindPositions();
+			if (SkeletalMeshAsset && FleshCollection && RestVertices)
+			{
+				using namespace UE::Chaos::FleshGenerator::Private;
+				BoundSurfacePositions(SkeletalMeshAsset, FleshCollection, RestVertices, RestVertices, Positions);
+
+				//@todo(Flesh Sections) : Add checks for multiple sections. 
+				int32 NumSections = MLDLOD.Sections.Num();
+				if (NumSections != 1)
+				{
+					UE_LOG(LogChaosFleshGeneratorPrivate, Warning, TEXT("SkeletalMeshAsset should have only one section."));
+					return None;
+				}
+
+				for (int32 SectionIndex = 0; SectionIndex < NumSections; ++SectionIndex)
+				{
+					const FSkelMeshSection& MLDSection = MLDLOD.Sections[SectionIndex];
+					if (MLDSection.NumVertices != Positions.Num())
+					{
+						UE_LOG(LogChaosFleshGeneratorPrivate, Warning, TEXT("SkeletalMeshAsset and FleshAsset have different number of vertices in section %d. Check if the assets have the same mesh."), SectionIndex);
+						return None;
+					}
+
+					for (int32 VertexIndex = 0; VertexIndex < MLDSection.NumVertices; ++VertexIndex)
+					{
+						const FVector3f& MLDPosition = MLDSection.SoftVertices[VertexIndex].Position;
+						const FVector3f& FleshPosition = Positions[VertexIndex];
+						if (!MLDPosition.Equals(FleshPosition, UE_KINDA_SMALL_NUMBER))
+						{
+							UE_LOG(LogChaosFleshGeneratorPrivate, Warning, TEXT("SkeletalMeshAsset and FleshAsset have different vertex positions. Check if the assets have the same vertex order."));
+							return None;
+						}
+					}
+				}
+			}
+
+			return Map;
+		}
+
+
+	};
+
+};
+
+#undef LOCTEXT_NAMESPACE
diff --git a/Engine/Plugins/Animation/MLDeformer/ChaosFleshGenerator/Source/ChaosFleshGenerator/Private/ChaosFleshGeneratorPrivate.h b/Engine/Plugins/Animation/MLDeformer/ChaosFleshGenerator/Source/ChaosFleshGenerator/Private/ChaosFleshGeneratorPrivate.h
new file mode 100644
--- /dev/null
+++ b/Engine/Plugins/Animation/MLDeformer/ChaosFleshGenerator/Source/ChaosFleshGenerator/Private/ChaosFleshGeneratorPrivate.h
@@ -0,0 +1,55 @@
+// Copyright Epic Games, Inc. All Rights Reserved.
+
+#pragma once
+
+#include "CoreTypes.h"
+#include "Containers/ContainersFwd.h"
+#include "GeometryCollection/ManagedArray.h"
+
+class FFleshCollection;
+class USkinnedAsset;
+class UFleshAsset;
+class UGeometryCache;
+class USkeletalMesh;
+class FSkeletalMeshLODRenderData;
+
+DECLARE_LOG_CATEGORY_EXTERN(LogChaosFleshGeneratorPrivate, Log, All);
+
+
+namespace UE::Chaos::FleshGenerator
+{
+	namespace Private
+	{
+		class FTimeScope
+		{
+			FString Name;
+			FDateTime StartTime;
+		public:
+			explicit FTimeScope(FString InName);
+			~FTimeScope();
+		};
+
+		void BoundSurfacePositions(
+			const USkeletalMesh* SkeletalMesh,
+			const FFleshCollection* FleshCollection,
+			const TManagedArray<FVector3f>* RestVertices,
+			const TManagedArray<FVector3f>* SimulatedVertices,
+			TArray<FVector3f>& Positions);
+
+		TArray<int32> ParseFrames(const FString& FramesString);
+
+		TArray<int32> Range(int32 End);
+
+		TArray<uint32> Range(uint32 Start, uint32 End);
+
+		int32 GetNumVertices(const FSkeletalMeshLODRenderData& LODData);
+
+		TArrayView<TArray<FVector3f>> ShrinkToValidFrames(const TArrayView<TArray<FVector3f>>& Positions, int32 NumVertices);
+
+		void SaveGeometryCache(UGeometryCache& GeometryCache, const USkinnedAsset& Asset, TConstArrayView<uint32> ImportedVertexNumbers, TArrayView<TArray<FVector3f>> PositionsToMoveFrom);
+
+		void SavePackage(UObject& Object);
+
+		TOptional<TArray<int32>> GetMeshImportVertexMap(const USkinnedAsset& SkinnedMeshAsset, const UFleshAsset& FleshAsset);
+	};
+};
\ No newline at end of file
diff --git a/Engine/Plugins/Animation/MLDeformer/ChaosFleshGenerator/Source/ChaosFleshGenerator/Private/ChaosFleshGeneratorSimulation.cpp b/Engine/Plugins/Animation/MLDeformer/ChaosFleshGenerator/Source/ChaosFleshGenerator/Private/ChaosFleshGeneratorSimulation.cpp
new file mode 100644
--- /dev/null
+++ b/Engine/Plugins/Animation/MLDeformer/ChaosFleshGenerator/Source/ChaosFleshGenerator/Private/ChaosFleshGeneratorSimulation.cpp
@@ -0,0 +1,182 @@
+// Copyright Epic Games, Inc. All Rights Reserved.
+
+#include "ChaosFleshGeneratorSimulation.h"
+
+#include "BonePose.h"
+#include "ChaosFlesh/FleshAsset.h"
+#include "ChaosFlesh/ChaosDeformableSolverComponent.h"
+#include "ChaosFleshGeneratorPrivate.h"
+#include "ChaosFleshGeneratorThreading.h"
+#include "FleshGeneratorComponent.h"
+#include "FleshGeneratorProperties.h"
+#include "Engine/SkeletalMesh.h"
+#include "Logging/LogMacros.h"
+
+DEFINE_LOG_CATEGORY(LogChaosFleshGeneratorSimulation);
+
+#define LOCTEXT_NAMESPACE "ChaosFleshGeneratorSimulation"
+
+namespace UE::Chaos::FleshGenerator
+{
+	FLaunchSimsTask::FLaunchSimsTask(TSharedPtr<FTaskResource> InTaskResource, TObjectPtr<UFleshGeneratorProperties> InProperties)
+	: TaskResource(InTaskResource)
+		, SimResources(InTaskResource->SimResources)
+		, Properties(InProperties)
+	{}
+
+	void FLaunchSimsTask::DoWork()
+	{	
+		const int32 NumFrames = TaskResource->FramesToSimulate.Num();
+		PrepareAnimationSequence();
+	
+		const int32 NumThreads = 1;
+	
+		for (int32 Frame = 0; Frame < NumFrames; Frame++)
+		{
+			if (!TaskResource->bCancelled.load())
+			{
+				const int32 ThreadIdx = Frame % NumThreads;
+				const int32 AnimFrame = TaskResource->FramesToSimulate[Frame];
+	
+				FSimResource& SimResource = *SimResources[ThreadIdx].Get();
+				SimResource.Pipe->Launch(*FString::Printf(TEXT("SimFrame:%d"), AnimFrame), [this, &SimResource, AnimFrame, Frame]()
+				{ 
+					FMemMark Mark(FMemStack::Get());
+					Simulate(SimResource, AnimFrame, Frame);
+				});
+			}
+			else
+			{
+				break;
+			}
+		}
+
+		for (TSharedPtr<FSimResource> SimResource : SimResources)
+		{
+			SimResource->Pipe->WaitUntilEmpty();
+		}
+	
+		RestoreAnimationSequence();
+	}
+
+	void FLaunchSimsTask::Simulate(FSimResource &SimResource, int32 AnimFrame, int32 CacheFrame) const
+	{
+		UFleshGeneratorComponent& FleshComponent = *SimResource.FleshComponent;
+		USkeletalGeneratorComponent& SkeletalComponent = *SimResource.SkeletalComponent;
+		UDeformableSolverComponent& SolverComponent = *SimResource.SolverComponent;
+	
+		const float DeltaTime = SolverComponent.SolverTiming.TimeStepSize;
+		const ESaveType SaveType = ESaveType::LastStep;
+	
+		const TArray<FTransform> Transforms = GetBoneTransforms( SkeletalComponent, AnimFrame);
+		FleshComponent.Pose(SkeletalComponent, Transforms);
+		SolverComponent.UpdateFromGameThread(DeltaTime);
+		SolverComponent.Simulate(DeltaTime);
+		SolverComponent.UpdateFromSimulation(DeltaTime);
+		SimResource.SimulatedPositions[CacheFrame] = GetRenderPositions(SimResource);
+
+		SimResource.FinishFrame();
+
+	}
+	
+	void FLaunchSimsTask::PrepareAnimationSequence()
+	{
+		TObjectPtr<UAnimSequence> AnimationSequence = Properties->AnimationSequence;
+		if (AnimationSequence)
+		{
+			InterpolationTypeBackup = AnimationSequence->Interpolation;
+			AnimationSequence->Interpolation = EAnimInterpolationType::Step;
+		}
+	}
+	
+	void FLaunchSimsTask::RestoreAnimationSequence()
+	{
+		TObjectPtr<UAnimSequence> AnimationSequence = Properties->AnimationSequence;
+		if (AnimationSequence)
+		{
+			AnimationSequence->Interpolation = InterpolationTypeBackup;
+		}
+	}
+	
+	TArray<FTransform> FLaunchSimsTask::GetBoneTransforms(USkeletalMeshComponent& InSkeletalComponent, int32 Frame) const
+	{
+		TArray<FTransform> ComponentSpaceTransforms;
+
+		const UAnimSequence* AnimationSequence = Properties->AnimationSequence;
+		const double Time = FMath::Clamp(AnimationSequence->GetSamplingFrameRate().AsSeconds(Frame), 0., (double)AnimationSequence->GetPlayLength());
+		FAnimExtractContext ExtractionContext(Time);
+
+		USkeletalMesh* SkeletalMesh = InSkeletalComponent.GetSkeletalMeshAsset();
+		FReferenceSkeleton* ReferenceSkeleton = &SkeletalMesh->GetRefSkeleton();
+		USkeleton* Skeleton = SkeletalMesh->GetSkeleton();
+		const int32 NumBones = ReferenceSkeleton ? ReferenceSkeleton->GetNum() : 0;
+	
+		TArray<uint16> BoneIndices;
+		BoneIndices.SetNumUninitialized(NumBones);
+		for (int32 Index = 0; Index < NumBones; ++Index)
+		{
+			BoneIndices[Index] = (uint16)Index;
+		}
+	
+		FBoneContainer BoneContainer;
+		BoneContainer.SetUseRAWData(true);
+		BoneContainer.InitializeTo(BoneIndices, UE::Anim::FCurveFilterSettings(), *Skeleton);
+	
+		FCompactPose OutPose;
+		OutPose.SetBoneContainer(&BoneContainer);
+		FBlendedCurve OutCurve;
+		OutCurve.InitFrom(BoneContainer);
+		UE::Anim::FStackAttributeContainer TempAttributes;
+	
+		FAnimationPoseData AnimationPoseData(OutPose, OutCurve, TempAttributes);
+		AnimationSequence->GetAnimationPose(AnimationPoseData, ExtractionContext);
+	
+		ComponentSpaceTransforms.SetNumUninitialized(NumBones);
+		for (int32 Index = 0; Index < NumBones; ++Index)
+		{
+			const FCompactPoseBoneIndex CompactIndex = BoneContainer.MakeCompactPoseIndex(FMeshPoseBoneIndex(Index));
+			const int32 ParentIndex = ReferenceSkeleton->GetParentIndex(Index);
+			if (ComponentSpaceTransforms.IsValidIndex(ParentIndex) && ParentIndex < Index)
+			{
+				ComponentSpaceTransforms[Index] = AnimationPoseData.GetPose()[CompactIndex] * ComponentSpaceTransforms[ParentIndex];
+			}
+			else
+			{
+				ComponentSpaceTransforms[Index] = ReferenceSkeleton->GetRefBonePose()[Index];
+			}
+		}
+		return ComponentSpaceTransforms;
+		
+	}
+	
+	TArray<FVector3f> FLaunchSimsTask::GetRenderPositions(FSimResource& SimResource) const
+	{
+		check(SimResource.FleshComponent);
+		check(SimResource.SkeletalComponent);
+
+		// @todo(brice) : Do we need this?
+		// This could potentially be slow. 
+		SimResource.SkeletalComponent->RecreateRenderState_Concurrent();
+		SimResource.bNeedsSkin.store(true);
+		SimResource.SkinEvent->Wait();
+
+		TArray<FVector3f> Positions;
+		const UFleshDynamicAsset* DynamicCollection = SimResource.FleshComponent->GetDynamicCollection();
+		const UFleshAsset* RestCollection = SimResource.FleshComponent->GetRestCollection();
+		const USkeletalMesh* SkeletalMesh = SimResource.SkeletalComponent->GetSkeletalMeshAsset();
+		if (RestCollection && DynamicCollection && SkeletalMesh)
+		{
+			const FFleshCollection* FleshCollection = RestCollection->GetCollection();
+			const TManagedArray<FVector3f>* RestVertices = RestCollection->FindPositions();
+			const TManagedArray<FVector3f>* SimulatedVertices = DynamicCollection->FindPositions();
+			if (FleshCollection && RestVertices && SimulatedVertices)
+			{
+				using namespace UE::Chaos::FleshGenerator::Private;
+				BoundSurfacePositions(SkeletalMesh, FleshCollection, RestVertices, SimulatedVertices, Positions);
+			}
+		}
+		return Positions;
+	}
+};
+
+#undef LOCTEXT_NAMESPACE
diff --git a/Engine/Plugins/Animation/MLDeformer/ChaosFleshGenerator/Source/ChaosFleshGenerator/Private/ChaosFleshGeneratorSimulation.h b/Engine/Plugins/Animation/MLDeformer/ChaosFleshGenerator/Source/ChaosFleshGenerator/Private/ChaosFleshGeneratorSimulation.h
new file mode 100644
--- /dev/null
+++ b/Engine/Plugins/Animation/MLDeformer/ChaosFleshGenerator/Source/ChaosFleshGenerator/Private/ChaosFleshGeneratorSimulation.h
@@ -0,0 +1,117 @@
+// Copyright Epic Games, Inc. All Rights Reserved.
+
+#pragma once
+
+#include "CoreTypes.h"
+#include "Tasks/Pipe.h"
+#include "UObject/ObjectMacros.h"
+#include "UObject/ObjectPtr.h"
+#include "Async/AsyncWork.h"
+#include "FleshGeneratorComponent.h"
+
+class UFleshGeneratorComponent;
+class USkeletalGeneratorComponent;
+class UDeformableSolverComponent;
+class UFleshGeneratorProperties;
+class FAsyncTaskNotification;
+class UGeometryCache;
+
+DECLARE_LOG_CATEGORY_EXTERN(LogChaosFleshGeneratorSimulation, Log, All);
+
+
+namespace UE::Chaos::FleshGenerator
+{
+	struct FTaskResource;
+
+	template<typename TaskType>
+	class TTaskRunner : public FNonAbandonableTask
+	{
+	public:
+		TTaskRunner(TUniquePtr<TaskType> InTask)
+			: Task(MoveTemp(InTask))
+		{
+		}
+
+		void DoWork()
+		{
+			if (Task)
+			{
+				Task->DoWork();
+			}
+		}
+
+		bool CanAbandon()
+		{
+			return true;
+		}
+
+		FORCEINLINE TStatId GetStatId() const
+		{
+			RETURN_QUICK_DECLARE_CYCLE_STAT(TTaskRunner, STATGROUP_ThreadPoolAsyncTasks);
+		}
+
+	private:
+		TUniquePtr<TaskType> Task;
+	};
+
+	struct FSimResource
+	{
+		TObjectPtr<UFleshGeneratorComponent> FleshComponent = nullptr;
+		TObjectPtr<USkeletalGeneratorComponent> SkeletalComponent = nullptr;
+		TObjectPtr<UDeformableSolverComponent> SolverComponent = nullptr;
+
+		TUniquePtr<::UE::Tasks::FPipe> Pipe;
+
+		FEvent* SkinEvent = nullptr;
+		std::atomic<bool> bNeedsSkin = false;
+
+		TArrayView<TArray<FVector3f>> SimulatedPositions;
+		std::atomic<int32>* NumSimulatedFrames = nullptr;
+
+		std::atomic<bool>* bCancelled = nullptr;
+
+		bool IsCancelled() const
+		{
+			return !bCancelled || bCancelled->load();
+		}
+		void FinishFrame()
+		{
+			if (NumSimulatedFrames)
+			{
+				++(*NumSimulatedFrames);
+			}
+		}
+	};
+
+	class FLaunchSimsTask
+	{
+	public:
+		FLaunchSimsTask(TSharedPtr<FTaskResource> InTaskResource, TObjectPtr<UFleshGeneratorProperties> InProperties);
+
+		void DoWork();
+
+	private:
+		using FPipe = ::UE::Tasks::FPipe;
+
+		enum class ESaveType
+		{
+			LastStep,
+			EveryStep,
+		};
+
+		void Simulate(FSimResource& SimResource, int32 AnimFrame, int32 CacheFrame) const;
+		void PrepareAnimationSequence();
+		void RestoreAnimationSequence();
+		TArray<FTransform> GetBoneTransforms(USkeletalMeshComponent& InSkeletalComponent, int32 Frame) const;
+		TArray<FVector3f> GetRenderPositions(FSimResource& SimResource) const;
+
+		TSharedPtr<FTaskResource> TaskResource;
+		TArray< TSharedPtr<FSimResource> >& SimResources;
+		TObjectPtr<UFleshGeneratorProperties> Properties;
+		EAnimInterpolationType InterpolationTypeBackup = EAnimInterpolationType::Linear;
+	};
+
+
+
+
+};
\ No newline at end of file
diff --git a/Engine/Plugins/Animation/MLDeformer/ChaosFleshGenerator/Source/ChaosFleshGenerator/Private/ChaosFleshGeneratorThreading.cpp b/Engine/Plugins/Animation/MLDeformer/ChaosFleshGenerator/Source/ChaosFleshGenerator/Private/ChaosFleshGeneratorThreading.cpp
new file mode 100644
--- /dev/null
+++ b/Engine/Plugins/Animation/MLDeformer/ChaosFleshGenerator/Source/ChaosFleshGenerator/Private/ChaosFleshGeneratorThreading.cpp
@@ -0,0 +1,140 @@
+// Copyright Epic Games, Inc. All Rights Reserved.
+
+#include "ChaosFleshGeneratorThreading.h"
+
+#include "ChaosFlesh/FleshAsset.h"
+#include "ChaosFlesh/ChaosDeformableSolverComponent.h"
+#include "ChaosFleshGenerator.h"
+#include "ChaosFleshGeneratorSimulation.h"
+#include "FleshGeneratorComponent.h"
+#include "FleshGeneratorProperties.h"
+#include "GeometryCache.h"
+#include "Logging/LogMacros.h"
+
+DEFINE_LOG_CATEGORY(LogChaosFleshGeneratorThreading);
+
+#define LOCTEXT_NAMESPACE "ChaosFleshGeneratorThreading"
+
+namespace UE::Chaos::FleshGenerator
+{
+	bool FTaskResource::AllocateSimResources_GameThread(TObjectPtr<UFleshGeneratorProperties> Properties, int32 Num)
+	{
+		UFleshAsset& Asset = *Properties->FleshAsset;
+		if (!ensure(Num == 1)) return false;
+
+		World = UWorld::CreateWorld(EWorldType::Editor, false);
+		SimResources.SetNum(Num);
+		for(int32 Index = 0; Index < Num; ++Index)
+		{
+			SimResources[Index] = TSharedPtr<FSimResource>(new FSimResource());
+			AActor* Owner = World->SpawnActor<AActor>();
+
+			UFleshGeneratorComponent* FleshComponent = NewObject<UFleshGeneratorComponent>(Owner);
+			FleshComponent->SetRestCollection(&Asset);
+			FleshComponent->RegisterComponentWithWorld(World);
+	
+			USkeletalGeneratorComponent* SkeletalMeshComponent = NewObject<USkeletalGeneratorComponent>(Owner);
+			SkeletalMeshComponent->SetSkeletalMesh(Asset.SkeletalMesh); 
+			SkeletalMeshComponent->RegisterComponentWithWorld(World);
+
+			UDeformableSolverComponent* DeformableSolver = NewObject<UDeformableSolverComponent>(Owner);
+			FleshComponent->EnableSimulation(DeformableSolver);
+			DeformableSolver->RegisterComponentWithWorld(World);
+
+			float TimeStepSize = (Properties->SolverTiming.FrameRate > 0) ? 1.0 / Properties->SolverTiming.FrameRate : 0;
+			DeformableSolver->SolverTiming.FixTimeStep = true;
+			DeformableSolver->SolverTiming.TimeStepSize = TimeStepSize;
+			DeformableSolver->SolverTiming.NumSubSteps = FMath::Max(0, Properties->SolverTiming.NumSubSteps);
+			DeformableSolver->SolverTiming.NumSolverIterations = FMath::Max(0, Properties->SolverTiming.NumIterations);
+			DeformableSolver->SolverEvolution = Properties->SolverEvolution;
+			DeformableSolver->SolverCollisions = Properties->SolverCollisions;
+			DeformableSolver->SolverConstraints = Properties->SolverConstraints;
+			DeformableSolver->SolverForces = Properties->SolverForces;
+			DeformableSolver->SolverDebugging = Properties->SolverDebugging;
+			DeformableSolver->Reset();
+			check(FleshComponent->GetPhysicsProxy());
+
+
+			constexpr int32 LODIndex = 0;
+			SkeletalMeshComponent->SetForcedLOD(LODIndex + 1);
+			SkeletalMeshComponent->UpdateLODStatus();
+			SkeletalMeshComponent->RefreshBoneTransforms(nullptr);
+
+			SkeletalMeshComponent->bRenderStatic = false;
+			constexpr bool bRecreateRenderStateImmediately = true;
+			SkeletalMeshComponent->SetCPUSkinningEnabled(true, bRecreateRenderStateImmediately);
+
+			FSimResource& SimResource = *SimResources[Index];
+			SimResource.FleshComponent = TObjectPtr<UFleshGeneratorComponent>(FleshComponent);
+			SimResource.SkeletalComponent = TObjectPtr<USkeletalGeneratorComponent>(SkeletalMeshComponent);
+			SimResource.SolverComponent = TObjectPtr<UDeformableSolverComponent>(DeformableSolver);
+			SimResource.Pipe = MakeUnique<UE::Tasks::FPipe>(*FString::Printf(TEXT("SimPipe:%d"), Index));
+			SimResource.SkinEvent = FPlatformProcess::GetSynchEventFromPool();
+			SimResource.bNeedsSkin.store(false);
+	
+			SimResource.SimulatedPositions = TArrayView<TArray<FVector3f>>(SimulatedPositions);
+			SimResource.NumSimulatedFrames = &NumSimulatedFrames;
+			SimResource.bCancelled = &bCancelled;
+	
+			if (FleshComponent == nullptr || SimResource.Pipe == nullptr)
+			{
+				UE_LOG(LogChaosFleshGenerator, Error, TEXT("Failed to allocate simulation resources"));
+				return false;
+			}
+		}
+
+		return true;
+	}
+	
+	void FTaskResource::FreeSimResources_GameThread()
+	{
+		if (Executer.IsValid())
+		{
+			Executer->EnsureCompletion();
+		}
+		for (TSharedPtr<FSimResource>& SimResource : SimResources)
+		{
+			FPlatformProcess::ReturnSynchEventToPool(SimResource->SkinEvent);
+			SimResource->Pipe.Reset();
+			SimResource->FleshComponent->UnregisterComponent();
+			SimResource->FleshComponent->DestroyComponent();
+		}
+		SimResources.Empty();
+		World->DestroyWorld(false);
+	}
+	
+	void FTaskResource::FlushRendering()
+	{
+		// Copy bNeedsSkin
+		TArray<bool> NeedsSkin;
+		NeedsSkin.SetNum(SimResources.Num());
+		bool bAnyNeedsSkin = false;
+		for (int32 Index = 0; Index < SimResources.Num(); ++Index)
+		{
+			const bool bNeedsSkin = SimResources[Index]->bNeedsSkin.load();
+			bAnyNeedsSkin |= bNeedsSkin;
+			NeedsSkin[Index] = bNeedsSkin;
+		}
+	
+		if (bAnyNeedsSkin)
+		{
+			FlushRenderingCommands();
+			for (int32 Index = 0; Index < SimResources.Num(); ++Index)
+			{
+				if (NeedsSkin[Index])
+				{
+					SimResources[Index]->bNeedsSkin.store(false);
+					SimResources[Index]->SkinEvent->Trigger();
+				}
+			}
+		}
+	}
+
+	void FTaskResource::Cancel()
+	{
+		bCancelled.store(true);
+		Executer->TryAbandonTask();
+	}
+};
+
+#undef LOCTEXT_NAMESPACE
diff --git a/Engine/Plugins/Animation/MLDeformer/ChaosFleshGenerator/Source/ChaosFleshGenerator/Private/ChaosFleshGeneratorThreading.h b/Engine/Plugins/Animation/MLDeformer/ChaosFleshGenerator/Source/ChaosFleshGenerator/Private/ChaosFleshGeneratorThreading.h
new file mode 100644
--- /dev/null
+++ b/Engine/Plugins/Animation/MLDeformer/ChaosFleshGenerator/Source/ChaosFleshGenerator/Private/ChaosFleshGeneratorThreading.h
@@ -0,0 +1,49 @@
+// Copyright Epic Games, Inc. All Rights Reserved.
+
+#pragma once
+
+#include "CoreTypes.h"
+#include "Logging/LogMacros.h"
+#include "ChaosFleshGeneratorSimulation.h"
+
+//#include "Async/AsyncWork.h"
+//#include "TickableEditorObject.h"
+//#include "UObject/StrongObjectPtr.h"
+
+class FAsyncTaskNotification;
+class UGeometryCache;
+struct FSimResource;
+
+DECLARE_LOG_CATEGORY_EXTERN(LogChaosFleshGeneratorThreading, Log, All);
+
+
+namespace UE::Chaos::FleshGenerator
+{
+	struct FTaskResource
+	{
+		using FExecuterType = FAsyncTask<TTaskRunner<FLaunchSimsTask>>;
+
+		TArray<TSharedPtr<FSimResource>> SimResources;
+
+		TUniquePtr<FExecuterType> Executer;
+		TUniquePtr<FAsyncTaskNotification> Notification;
+		FDateTime StartTime;
+		FDateTime LastUpdateTime;
+
+		TArray<int32> FramesToSimulate;
+		TArray<TArray<FVector3f>> SimulatedPositions;
+		TArray<uint32> ImportedVertexNumbers;
+		UGeometryCache* Cache = nullptr;
+
+		std::atomic<int32> NumSimulatedFrames = 0;
+		std::atomic<bool> bCancelled = false;
+
+		UWorld* World = nullptr;
+
+		bool AllocateSimResources_GameThread(TObjectPtr<UFleshGeneratorProperties> Properties, int32 Num);
+		void FreeSimResources_GameThread();
+		void FlushRendering();
+		void Cancel();
+	};
+
+};
\ No newline at end of file
diff --git a/Engine/Plugins/Animation/MLDeformer/ChaosFleshGenerator/Source/ChaosFleshGenerator/Private/ChaosFleshGeneratorToolsMenuExtender.cpp b/Engine/Plugins/Animation/MLDeformer/ChaosFleshGenerator/Source/ChaosFleshGenerator/Private/ChaosFleshGeneratorToolsMenuExtender.cpp
new file mode 100644
--- /dev/null
+++ b/Engine/Plugins/Animation/MLDeformer/ChaosFleshGenerator/Source/ChaosFleshGenerator/Private/ChaosFleshGeneratorToolsMenuExtender.cpp
@@ -0,0 +1,120 @@
+// Copyright Epic Games, Inc. All Rights Reserved.
+#include "ChaosFleshGeneratorToolsMenuExtender.h"
+
+#include "FleshGeneratorProperties.h"
+#include "Engine/SkeletalMesh.h"
+#include "Framework/Docking/TabManager.h"
+#include "IDocumentation.h"
+#include "MLDeformerGeomCacheEditorModel.h"
+#include "MLDeformerEditorModel.h"
+#include "MLDeformerModel.h"
+#include "MLDeformerGeomCacheModel.h"
+#include "SFleshGeneratorWidget.h"
+#include "Widgets/Docking/SDockTab.h"
+
+#define LOCTEXT_NAMESPACE "ChaosFleshGeneratorToolsMenuExtender"
+
+namespace UE::Chaos::FleshGenerator
+{
+	namespace Private
+	{
+		void SpawnTab(UE::MLDeformer::FMLDeformerEditorToolkit& Toolkit)
+		{
+			TSharedPtr<SDockTab> Tab = Toolkit.GetAssociatedTabManager()->TryInvokeTab(FName(FChaosFleshGeneratorTabSummoner::TabID));
+			if (!Tab.IsValid())
+			{
+				return;
+			}
+			
+			TSharedRef<SWidget> TabContent = Tab->GetContent();
+			if (TabContent == SNullWidget::NullWidget || TabContent->GetTypeAsString() != "SFleshGeneratorWidget")
+			{
+				return;
+			}
+			TSharedRef<SFleshGeneratorWidget> Widget = StaticCastSharedRef<SFleshGeneratorWidget>(Tab->GetContent());
+			if (Widget == SNullWidget::NullWidget)
+			{
+				return;
+			}
+		
+			TWeakObjectPtr<UFleshGeneratorProperties> Properties = Widget->GetProperties();
+			if (!Properties.IsValid())
+			{
+				return;
+			}
+
+			using UE::MLDeformer::FMLDeformerEditorModel;
+			if (FMLDeformerEditorModel* EditorModel = Toolkit.GetActiveModel())
+			{
+				if (UMLDeformerModel* Model = EditorModel->GetModel())
+				{
+					Properties->SkeletalMeshAsset = Model->GetSkeletalMesh();
+					Properties->AnimationSequence = EditorModel->GetActiveTrainingInputAnimSequence();
+					if (UMLDeformerGeomCacheModel* GeomCacheModel = Cast<UMLDeformerGeomCacheModel>(Model))
+					{
+						const UE::MLDeformer::FMLDeformerGeomCacheEditorModel* GeomCacheEditorModel = static_cast<const UE::MLDeformer::FMLDeformerGeomCacheEditorModel*>(EditorModel);
+						if (UGeometryCache* GeomCache = GeomCacheEditorModel->GetActiveGeometryCache())
+						{
+							Properties->SimulatedCache = GeomCache;
+						}
+					}
+				}
+			}
+		}
+	};
+	
+	TUniquePtr<FChaosFleshGeneratorToolsMenuExtender> CreateToolsMenuExtender()
+	{
+		return MakeUnique<FChaosFleshGeneratorToolsMenuExtender>();
+	}
+	
+	const FName FChaosFleshGeneratorTabSummoner::TabID = FName("ChaosFleshGenerator");
+	
+	FChaosFleshGeneratorTabSummoner::FChaosFleshGeneratorTabSummoner(const TSharedRef<FMLDeformerEditorToolkit>& InEditor)
+		: FWorkflowTabFactory(TabID, InEditor)
+	{
+		bIsSingleton = true;
+		TabLabel = LOCTEXT("ChaosFleshGenerator", "Chaos Flesh Generator");
+		ViewMenuDescription = LOCTEXT("ViewMenu_Desc", "Chaos Flesh Generator");
+		ViewMenuTooltip = LOCTEXT("ViewMenu_ToolTip", "Show the Chaos Flesh Generator tab.");
+	}
+	
+	TSharedRef<SWidget> FChaosFleshGeneratorTabSummoner::CreateTabBody(const FWorkflowTabSpawnInfo& Info) const
+	{
+		return SNew(SFleshGeneratorWidget);
+	}
+	
+	TSharedPtr<SToolTip> FChaosFleshGeneratorTabSummoner::CreateTabToolTipWidget(const FWorkflowTabSpawnInfo& Info) const
+	{
+		return IDocumentation::Get()->CreateToolTip(
+			LOCTEXT("ChaosFleshGeneratorWidgetTooltip", "Generate training data using chaos cloth solver."),
+			NULL,
+			TEXT("Shared/Editors/Persona"),
+			TEXT("ChaosFleshGenerator_Window"));
+	}
+	
+	FMenuEntryParams FChaosFleshGeneratorToolsMenuExtender::GetMenuEntry(FMLDeformerEditorToolkit& Toolkit) const
+	{
+		FMenuEntryParams Params;
+		Params.DirectActions = FUIAction(
+			FExecuteAction::CreateLambda([&Toolkit]()
+			{
+				Private::SpawnTab(Toolkit);
+			}),
+			FCanExecuteAction::CreateLambda([]()
+			{
+				return true;
+			})
+		);
+		Params.LabelOverride = LOCTEXT("ChaosFleshGenerator", "Chaos Flesh Generator");
+		Params.ToolTipOverride = LOCTEXT("ChaosFleshGeneratorMenuTooltip", "Generate training data using chaos cloth solver");
+		return Params;
+	}
+	
+	TSharedPtr<FWorkflowTabFactory> FChaosFleshGeneratorToolsMenuExtender::GetTabSummoner(const TSharedRef<FMLDeformerEditorToolkit>& Toolkit) const
+	{
+		return MakeShared<FChaosFleshGeneratorTabSummoner>(Toolkit);
+	}
+};
+
+#undef LOCTEXT_NAMESPACE
\ No newline at end of file
diff --git a/Engine/Plugins/Animation/MLDeformer/ChaosFleshGenerator/Source/ChaosFleshGenerator/Private/ChaosFleshGeneratorToolsMenuExtender.h b/Engine/Plugins/Animation/MLDeformer/ChaosFleshGenerator/Source/ChaosFleshGenerator/Private/ChaosFleshGeneratorToolsMenuExtender.h
new file mode 100644
--- /dev/null
+++ b/Engine/Plugins/Animation/MLDeformer/ChaosFleshGenerator/Source/ChaosFleshGenerator/Private/ChaosFleshGeneratorToolsMenuExtender.h
@@ -0,0 +1,32 @@
+// Copyright Epic Games, Inc. All Rights Reserved.
+
+#pragma once
+
+#include "CoreTypes.h"
+#include "MLDeformerEditorToolkit.h"
+#include "WorkflowOrientedApp/WorkflowTabFactory.h"
+
+namespace UE::Chaos::FleshGenerator
+{
+	struct FChaosFleshGeneratorTabSummoner : public FWorkflowTabFactory
+	{
+	public:
+		static const FName TabID;
+		using FMLDeformerEditorToolkit = UE::MLDeformer::FMLDeformerEditorToolkit;
+	
+		FChaosFleshGeneratorTabSummoner(const TSharedRef<FMLDeformerEditorToolkit>& InEditor);
+		virtual TSharedRef<SWidget> CreateTabBody(const FWorkflowTabSpawnInfo& Info) const override;
+		virtual TSharedPtr<SToolTip> CreateTabToolTipWidget(const FWorkflowTabSpawnInfo& Info) const override;
+	};
+	
+	class FChaosFleshGeneratorToolsMenuExtender : public UE::MLDeformer::FToolsMenuExtender
+	{
+	public:
+		using FMLDeformerEditorToolkit = UE::MLDeformer::FMLDeformerEditorToolkit;
+		~FChaosFleshGeneratorToolsMenuExtender() override = default;
+		virtual FMenuEntryParams GetMenuEntry(FMLDeformerEditorToolkit& Toolkit) const override;
+		virtual TSharedPtr<FWorkflowTabFactory> GetTabSummoner(const TSharedRef<FMLDeformerEditorToolkit>& Toolkit) const;
+	};
+
+	TUniquePtr<FChaosFleshGeneratorToolsMenuExtender> CreateToolsMenuExtender();
+};
\ No newline at end of file
diff --git a/Engine/Plugins/Animation/MLDeformer/ChaosFleshGenerator/Source/ChaosFleshGenerator/Private/FleshGeneratorComponent.cpp b/Engine/Plugins/Animation/MLDeformer/ChaosFleshGenerator/Source/ChaosFleshGenerator/Private/FleshGeneratorComponent.cpp
new file mode 100644
--- /dev/null
+++ b/Engine/Plugins/Animation/MLDeformer/ChaosFleshGenerator/Source/ChaosFleshGenerator/Private/FleshGeneratorComponent.cpp
@@ -0,0 +1,56 @@
+// Copyright Epic Games, Inc. All Rights Reserved.
+
+#include "FleshGeneratorComponent.h"
+
+
+//
+// UFleshGeneratorComponent
+//
+
+
+UFleshGeneratorComponent::UFleshGeneratorComponent(const FObjectInitializer& ObjectInitializer)
+    : Super(ObjectInitializer)
+{
+}
+
+UFleshGeneratorComponent::~UFleshGeneratorComponent() = default;
+
+void UFleshGeneratorComponent::TickComponent(float DeltaTime, enum ELevelTick TickType, FActorComponentTickFunction* ThisTickFunction)
+{
+    UActorComponent::TickComponent(DeltaTime, TickType, ThisTickFunction);
+}
+
+void UFleshGeneratorComponent::Pose(USkeletalGeneratorComponent& InSkeletalComponent, const TArray<FTransform>& InComponentSpaceTransforms)
+{
+	if (!ensure(InComponentSpaceTransforms.Num() == InSkeletalComponent.GetComponentSpaceTransforms().Num()))
+	{
+		return;
+	}
+	InSkeletalComponent.GetEditableComponentSpaceTransforms() = InComponentSpaceTransforms;
+	InSkeletalComponent.FlipSpaceBuffer();
+	InSkeletalComponent.FinalizeBoneTransform();
+
+	InSkeletalComponent.UpdateBounds();
+    if (IsInGameThread())
+    {
+		InSkeletalComponent.MarkRenderTransformDirty();
+		InSkeletalComponent.MarkRenderDynamicDataDirty();
+    }
+}
+
+//
+// USkeletalGeneratorComponent
+//
+
+USkeletalGeneratorComponent::USkeletalGeneratorComponent(const FObjectInitializer& ObjectInitializer)
+	: Super(ObjectInitializer)
+{
+}
+
+USkeletalGeneratorComponent::~USkeletalGeneratorComponent() = default;
+
+void USkeletalGeneratorComponent::FlipSpaceBuffer()
+{
+	bNeedToFlipSpaceBaseBuffers = true;
+
+}
diff --git a/Engine/Plugins/Animation/MLDeformer/ChaosFleshGenerator/Source/ChaosFleshGenerator/Private/FleshGeneratorComponent.h b/Engine/Plugins/Animation/MLDeformer/ChaosFleshGenerator/Source/ChaosFleshGenerator/Private/FleshGeneratorComponent.h
new file mode 100644
--- /dev/null
+++ b/Engine/Plugins/Animation/MLDeformer/ChaosFleshGenerator/Source/ChaosFleshGenerator/Private/FleshGeneratorComponent.h
@@ -0,0 +1,47 @@
+// Copyright Epic Games, Inc. All Rights Reserved.
+
+#pragma once
+
+#include "Chaos/Deformable/ChaosDeformableSolverProxy.h"
+#include "ChaosFlesh/FleshComponent.h"
+#include "Components/SkeletalMeshComponent.h"
+#include "ChaosFlesh/ChaosDeformableTetrahedralComponent.h"
+
+
+#include "FleshGeneratorComponent.generated.h"
+
+class USkeletalGeneratorComponent;
+
+
+/**
+ * Flesh data generation component.
+ */
+UCLASS()
+class UFleshGeneratorComponent : public UFleshComponent
+{
+	GENERATED_BODY()	
+public:
+	UFleshGeneratorComponent(const FObjectInitializer& ObjectInitializer);
+	~UFleshGeneratorComponent();
+
+	/** Pose the cloth component using component space transforms. */
+    void Pose(USkeletalGeneratorComponent& InSkeletalComponent, const TArray<FTransform>& InComponentSpaceTransforms);
+protected:
+	//~ Begin UActorComponent Interface
+	virtual void TickComponent(float DeltaTime, enum ELevelTick TickType, FActorComponentTickFunction* ThisTickFunction) override;
+	//~ End UActorComponent Interface
+};
+
+
+/**
+ * USkeletalGeneratorComponent data generation component.
+ */
+UCLASS()
+class USkeletalGeneratorComponent : public USkeletalMeshComponent
+{
+	GENERATED_BODY()
+public:
+	USkeletalGeneratorComponent(const FObjectInitializer& ObjectInitializer);
+	~USkeletalGeneratorComponent();
+	void FlipSpaceBuffer();
+};
diff --git a/Engine/Plugins/Animation/MLDeformer/ChaosFleshGenerator/Source/ChaosFleshGenerator/Private/FleshGeneratorProperties.h b/Engine/Plugins/Animation/MLDeformer/ChaosFleshGenerator/Source/ChaosFleshGenerator/Private/FleshGeneratorProperties.h
new file mode 100644
--- /dev/null
+++ b/Engine/Plugins/Animation/MLDeformer/ChaosFleshGenerator/Source/ChaosFleshGenerator/Private/FleshGeneratorProperties.h
@@ -0,0 +1,82 @@
+// Copyright Epic Games, Inc. All Rights Reserved.
+
+#pragma once
+
+#include "CoreMinimal.h"
+
+#include "ChaosFlesh/FleshAsset.h"
+#include "ChaosFlesh/ChaosDeformableSolverComponent.h"
+#include "Animation/AnimSequence.h"
+#include "Engine/SkinnedAsset.h"
+#include "GeometryCache.h"
+
+#include "FleshGeneratorProperties.generated.h"
+
+
+USTRUCT(BlueprintType)
+struct FFleshGeneratorSolverTimingGroup
+{
+	GENERATED_USTRUCT_BODY()
+
+		/* Frame rate of the training animation (FrameDeltaTime = 1/FrameRate) */
+		UPROPERTY(EditAnywhere, Category = "SolverTiming", meta = (Min = 1))
+		float FrameRate = 24.;
+
+	/* Number of frames to simulate ( the length of the animation) */
+	UPROPERTY(EditAnywhere, Category = "SolverTiming", meta = (Min = 0))
+		int32 NumFrames = 150;
+
+	/* Number of subdivisions within a time step (SolverStep = FrameDeltaTime/NumSubSteps) */
+	UPROPERTY(EditAnywhere, Category = "SolverTiming", meta = (Min = 0))
+		int32 NumSubSteps = 2;
+
+	/* Number of convergence steps for the constraint solver for each solver step. */
+	UPROPERTY(EditAnywhere, Category = "SolverTiming", meta = (Min = 0))
+		int32 NumIterations = 5;
+
+};
+
+
+UCLASS()
+class UFleshGeneratorProperties : public UObject
+{
+	GENERATED_BODY()
+public:
+	/* Skeletal mesh that will be used in MLDeformer */
+	UPROPERTY(EditAnywhere, Category = Input)
+	TObjectPtr<USkinnedAsset> SkeletalMeshAsset;
+
+	/* Chaos cloth asset used in simulation. This should be different from the skeletal mesh asset. */
+	UPROPERTY(EditAnywhere, Category = Input)
+	TObjectPtr<UFleshAsset> FleshAsset;
+
+	/* Training poses. */
+	UPROPERTY(EditAnywhere, Category = Input)
+	TObjectPtr<UAnimSequence> AnimationSequence;
+
+	/* e.g. "0, 2, 5-10, 12-15". If left empty, all frames will be used */
+	UPROPERTY(EditAnywhere, Category = Input)
+	FString FramesToSimulate;
+
+	/* Output meshes */
+	UPROPERTY(EditAnywhere, Category = Output)
+	TObjectPtr<UGeometryCache> SimulatedCache;
+
+	UPROPERTY(EditAnywhere, BlueprintReadOnly, Category = "SimulationSettings")
+	FFleshGeneratorSolverTimingGroup SolverTiming;
+
+	UPROPERTY(EditAnywhere, BlueprintReadOnly, Category = "SimulationSettings")
+	FSolverEvolutionGroup SolverEvolution;
+
+	UPROPERTY(EditAnywhere, BlueprintReadOnly, Category = "SimulationSettings")
+	FSolverCollisionsGroup SolverCollisions;
+
+	UPROPERTY(EditAnywhere, BlueprintReadOnly, Category = "SimulationSettings")
+	FSolverConstraintsGroup SolverConstraints;
+
+	UPROPERTY(EditAnywhere, BlueprintReadOnly, Category = "SimulationSettings")
+	FSolverForcesGroup SolverForces;
+
+	UPROPERTY(EditAnywhere, BlueprintReadOnly, Category = "SimulationSettings")
+	FSolverDebuggingGroup SolverDebugging;
+};
\ No newline at end of file
diff --git a/Engine/Plugins/Animation/MLDeformer/ChaosFleshGenerator/Source/ChaosFleshGenerator/Private/SFleshGeneratorWidget.cpp b/Engine/Plugins/Animation/MLDeformer/ChaosFleshGenerator/Source/ChaosFleshGenerator/Private/SFleshGeneratorWidget.cpp
new file mode 100644
--- /dev/null
+++ b/Engine/Plugins/Animation/MLDeformer/ChaosFleshGenerator/Source/ChaosFleshGenerator/Private/SFleshGeneratorWidget.cpp
@@ -0,0 +1,188 @@
+// Copyright Epic Games, Inc. All Rights Reserved.
+
+#include "SFleshGeneratorWidget.h"
+
+#include "AssetRegistry/AssetRegistryModule.h"
+#include "FleshGeneratorProperties.h"
+#include "ChaosFleshGenerator.h"
+#include "ContentBrowserModule.h"
+#include "DetailCategoryBuilder.h"
+#include "DetailLayoutBuilder.h"
+#include "IContentBrowserSingleton.h"
+#include "IDetailsView.h"
+#include "Misc/FileHelper.h"
+#include "PropertyCustomizationHelpers.h"
+#include "PropertyEditorModule.h"
+#include "Widgets/Input/SButton.h"
+#include "Widgets/Layout/SBox.h"
+#include "Widgets/Text/STextBlock.h"
+
+#define LOCTEXT_NAMESPACE "FleshGeneratorWidget"
+
+namespace UE::Chaos::FleshGenerator
+{
+	void SFleshGeneratorWidget::Construct(const FArguments& InArgs)
+	{
+		FPropertyEditorModule& PropertyModule = FModuleManager::LoadModuleChecked<FPropertyEditorModule>("PropertyEditor");
+		ChaosFleshGenerator = MakeShared<FChaosFleshGenerator>();
+
+		FDetailsViewArgs Args;
+		DetailsView = PropertyModule.CreateDetailView(Args);
+		DetailsView->SetObject(&ChaosFleshGenerator->GetProperties());
+
+		this->ChildSlot
+		[
+			SNew(SVerticalBox)
+			+ SVerticalBox::Slot()
+			.AutoHeight()
+			[
+				DetailsView.ToSharedRef()
+			]
+			+ SVerticalBox::Slot()
+			.AutoHeight()
+			.HAlign(HAlign_Center)
+			[
+				SNew(SBox)
+				.WidthOverride(300)
+				[
+					SNew(SButton)
+					.Text(FText::FromString("Start Generating"))
+					.HAlign(HAlign_Center)
+					.OnClicked_Lambda([this]() -> FReply
+					{
+						ChaosFleshGenerator->RequestAction(EFleshGeneratorActions::StartGenerate);
+						return FReply::Handled();
+					})
+				]
+			]
+		];
+	}
+
+	TWeakObjectPtr<UFleshGeneratorProperties> SFleshGeneratorWidget::GetProperties() const
+	{
+		if (ChaosFleshGenerator.IsValid())
+		{
+			return &ChaosFleshGenerator->GetProperties();
+		}
+		return nullptr;
+	}
+
+	TSharedRef<IDetailCustomization> FFleshGeneratorDetails::MakeInstance()
+	{
+		return MakeShareable(new FFleshGeneratorDetails);
+	}
+
+	namespace Private
+	{
+		template<class T>
+		T* CreateOrLoad(const FString& PackageName)
+		{
+			const FName AssetName(FPackageName::GetLongPackageAssetName(PackageName));		
+			if (UPackage* const Package = CreatePackage(*PackageName))
+			{
+				LoadPackage(nullptr, *PackageName, LOAD_Quiet | LOAD_EditorOnly);
+				T* Asset = FindObject<T>(Package, *AssetName.ToString());
+				if (!Asset)
+				{
+					Asset = NewObject<T>(Package, *AssetName.ToString(), RF_Public | RF_Standalone | RF_Transactional);
+					Asset->MarkPackageDirty();
+					FAssetRegistryModule::AssetCreated(Asset);
+				}
+				return Asset;
+			}
+			return nullptr;
+		}
+
+		TObjectPtr<UGeometryCache> NewGeometryCacheDialog(const UObject* NamingAsset = nullptr)
+		{
+			FSaveAssetDialogConfig Config;
+			{
+				if (NamingAsset)
+				{
+					const FString PackageName = NamingAsset->GetOutermost()->GetName();
+					Config.DefaultPath = FPackageName::GetLongPackagePath(PackageName);
+					Config.DefaultAssetName = FString::Printf(TEXT("GC_%s"), *NamingAsset->GetName());
+				}
+				Config.AssetClassNames.Add(UGeometryCache::StaticClass()->GetClassPathName());
+				Config.ExistingAssetPolicy = ESaveAssetDialogExistingAssetPolicy::Disallow;
+				Config.DialogTitleOverride = LOCTEXT("ExportGeometryCacheDialogTitle", "Export Geometry Cache As");
+			}
+				
+			FContentBrowserModule& ContentBrowserModule = FModuleManager::LoadModuleChecked<FContentBrowserModule>("ContentBrowser");
+				
+			FString NewPackageName;
+			FText OutError;
+			for (bool bFilenameValid = false; !bFilenameValid; bFilenameValid = FFileHelper::IsFilenameValidForSaving(NewPackageName, OutError))
+			{
+				const FString AssetPath = ContentBrowserModule.Get().CreateModalSaveAssetDialog(Config);
+				if (AssetPath.IsEmpty())
+				{
+					return nullptr;
+				}
+				NewPackageName = FPackageName::ObjectPathToPackageName(AssetPath);
+			}
+			return CreateOrLoad<UGeometryCache>(NewPackageName);
+		}
+
+		TWeakObjectPtr<UFleshGeneratorProperties> GetProperties(IDetailLayoutBuilder& DetailBuilder)
+		{
+			TArray<TWeakObjectPtr<UObject>> ObjectsBeingCustomized;
+			DetailBuilder.GetObjectsBeingCustomized(ObjectsBeingCustomized);
+			if (ObjectsBeingCustomized.IsEmpty())
+			{
+				return nullptr;
+			}
+			else
+			{
+				return Cast<UFleshGeneratorProperties>(ObjectsBeingCustomized[0]);
+			}
+		}
+
+		void AddGeometryCacheRowWithButton(IDetailLayoutBuilder& DetailBuilder, IDetailCategoryBuilder& CategoryBuilder, const FName PropertyName)
+		{
+			TSharedPtr<IPropertyHandle> Property = DetailBuilder.GetProperty(PropertyName);
+			DetailBuilder.HideProperty(Property);
+			CategoryBuilder.AddCustomRow(FText::FromName(PropertyName))
+			.NameContent()
+			[
+				Property->CreatePropertyNameWidget()
+			]
+			.ValueContent()
+			[
+				SNew(SHorizontalBox)
+				+ SHorizontalBox::Slot()
+				.AutoWidth()
+				[
+					Property->CreatePropertyValueWidget()
+				]
+				+ SHorizontalBox::Slot()
+				.MaxWidth(100)
+				.VAlign(VAlign_Center)
+				[
+					SNew(SButton)
+					.Text(FText::FromString("New"))
+					.HAlign(HAlign_Center)
+					.VAlign(VAlign_Center)
+					.OnClicked_Lambda([Property, &DetailBuilder]() -> FReply
+					{
+						TWeakObjectPtr<UFleshGeneratorProperties> Properties = Private::GetProperties(DetailBuilder);
+						const UObject* const NamingAsset = Properties.IsValid() ? Properties->SkeletalMeshAsset.Get() : nullptr;
+						UGeometryCache* const NewGeometryCache = Private::NewGeometryCacheDialog(NamingAsset);
+						Property->SetValue(NewGeometryCache);
+						return FReply::Handled();
+					})
+				]
+			];
+		}
+	}
+
+
+	void FFleshGeneratorDetails::CustomizeDetails(IDetailLayoutBuilder& DetailBuilder)
+	{
+		IDetailCategoryBuilder& InputCategory = DetailBuilder.EditCategory("Input");
+		IDetailCategoryBuilder& OutputCategory = DetailBuilder.EditCategory("Output");
+		Private::AddGeometryCacheRowWithButton(DetailBuilder, OutputCategory, GET_MEMBER_NAME_CHECKED(UFleshGeneratorProperties, SimulatedCache));
+	}
+};
+
+#undef LOCTEXT_NAMESPACE
\ No newline at end of file
diff --git a/Engine/Plugins/Animation/MLDeformer/ChaosFleshGenerator/Source/ChaosFleshGenerator/Private/SFleshGeneratorWidget.h b/Engine/Plugins/Animation/MLDeformer/ChaosFleshGenerator/Source/ChaosFleshGenerator/Private/SFleshGeneratorWidget.h
new file mode 100644
--- /dev/null
+++ b/Engine/Plugins/Animation/MLDeformer/ChaosFleshGenerator/Source/ChaosFleshGenerator/Private/SFleshGeneratorWidget.h
@@ -0,0 +1,37 @@
+// Copyright Epic Games, Inc. All Rights Reserved.
+
+#pragma once
+
+#include "IDetailCustomization.h"
+#include "Widgets/SCompoundWidget.h"
+
+class UFleshGeneratorProperties;
+class IDetailsView;
+namespace UE::Chaos::FleshGenerator
+{
+	class FChaosFleshGenerator;
+
+	class SFleshGeneratorWidget : public SCompoundWidget
+	{
+		SLATE_BEGIN_ARGS(SFleshGeneratorWidget) 
+		{}
+		SLATE_END_ARGS()
+
+	public:
+		void Construct(const FArguments& InArgs);
+		TWeakObjectPtr<UFleshGeneratorProperties> GetProperties() const;
+
+	private:
+		TSharedPtr<IDetailsView> DetailsView;
+		TSharedPtr<FChaosFleshGenerator> ChaosFleshGenerator;
+	};
+
+	class FFleshGeneratorDetails : public IDetailCustomization
+	{
+	public:
+		static TSharedRef<IDetailCustomization> MakeInstance();
+		virtual void CustomizeDetails(IDetailLayoutBuilder& DetailBuilder) override;
+	private:
+		TWeakObjectPtr<UFleshGeneratorProperties> GetProperties(IDetailLayoutBuilder& DetailBuilder) const;
+	};
+};
\ No newline at end of file
diff --git a/Engine/Source/Runtime/Experimental/Chaos/Private/GeometryCollection/Facades/CollectionTetrahedralBindingsFacade.cpp b/Engine/Source/Runtime/Experimental/Chaos/Private/GeometryCollection/Facades/CollectionTetrahedralBindingsFacade.cpp
--- a/Engine/Source/Runtime/Experimental/Chaos/Private/GeometryCollection/Facades/CollectionTetrahedralBindingsFacade.cpp
+++ b/Engine/Source/Runtime/Experimental/Chaos/Private/GeometryCollection/Facades/CollectionTetrahedralBindingsFacade.cpp
@@ -54,10 +54,10 @@ namespace GeometryCollection::Facades
 		const FVector3f& CurrPtC)
 	{
 		Chaos::PMatrix<float, 3, 3> BasisVectors = GetOrthogonalBasisVectors(RestPtA, RestPtB, RestPtC);
-		Chaos::PMatrix<float, 3, 3> RestRotInv = BasisVectors.GetTransposed(); // Ortho matrix, so T == I
+		Chaos::PMatrix<float, 3, 3> RestRotInv = BasisVectors.Inverse();
 		Chaos::PMatrix<float, 3, 3> CurrRot = GetOrthogonalBasisVectors(CurrPtA, CurrPtB, CurrPtC);
 		Chaos::PMatrix<float, 3, 3> BasisDelta = RestRotInv * CurrRot;
-		return BasisDelta.TransformVector(Offset); // ryan - mult order?
+		return BasisDelta.TransformVector(Offset);
 	}
 
 	FVector3f
@@ -72,6 +72,7 @@ namespace GeometryCollection::Facades
 			RestVertices[Parents[0]], RestVertices[Parents[1]], RestVertices[Parents[2]],
 			CurrVertices[Parents[0]], CurrVertices[Parents[1]], CurrVertices[Parents[2]]);
 	}
+	FString bump;
 
 	FVector3f
 	FleshDeformerImpl::GetEmbeddedPosition(
@@ -85,19 +86,23 @@ namespace GeometryCollection::Facades
 		const FIntVector4& Parents = ParentsArray->Get()[SurfaceIndex];
 		const FVector4f& Weights = WeightsArray->Get()[SurfaceIndex];
 
-		FVector3f Pos(0, 0, 0);
+		FVector3f Pos = FVector3f::Zero();
 
 		const int32 iEnd = Parents[3] == INDEX_NONE ? 3 : 4;
 		for (int32 i = 0; i < iEnd; i++)
 		{
 			Pos += CurrVertices[Parents[i]] * Weights[i];
 		}
+		FVector3f Pos2 = Pos;
+		FString bumpe = FString::Printf(TEXT("(%f,%f,%f)"), Pos.X, Pos.Y, Pos.Z);
+		bump = bumpe;
 
 		// If surface binding, the last index is -1.
 		if (Parents[3] == INDEX_NONE)
 		{
 			const FVector3f& Offset = OffsetArray->Get()[SurfaceIndex];
-			Pos += GetRotatedOffsetVector(Parents, Offset, RestVertices, CurrVertices);
+			const FVector3f RotatedOffset = GetRotatedOffsetVector(Parents, Offset, RestVertices, CurrVertices);
+			Pos -= Offset;
 		}
 
 		return Pos;
