ChaosFlesh : Static CIS fix
- Remvoed skin event callback
- Removed Pipe support.

diff --git a/Engine/Plugins/Animation/MLDeformer/ChaosFleshGenerator/Source/ChaosFleshGenerator/Private/ChaosFleshGeneratorSimulation.cpp b/Engine/Plugins/Animation/MLDeformer/ChaosFleshGenerator/Source/ChaosFleshGenerator/Private/ChaosFleshGeneratorSimulation.cpp
--- a/Engine/Plugins/Animation/MLDeformer/ChaosFleshGenerator/Source/ChaosFleshGenerator/Private/ChaosFleshGeneratorSimulation.cpp
+++ b/Engine/Plugins/Animation/MLDeformer/ChaosFleshGenerator/Source/ChaosFleshGenerator/Private/ChaosFleshGeneratorSimulation.cpp
@@ -35,15 +35,11 @@ namespace UE::Chaos::FleshGenerator
 		{
 			if (!TaskResource->bCancelled.load())
 			{
-				const int32 ThreadIdx = Frame % NumThreads;
 				const int32 AnimFrame = TaskResource->FramesToSimulate[Frame];
 	
+				const int32 ThreadIdx = 0;// Frame% NumThreads;
 				FSimResource& SimResource = *SimResources[ThreadIdx].Get();
-				SimResource.Pipe->Launch(*FString::Printf(TEXT("SimFrame:%d"), AnimFrame), [this, &SimResource, AnimFrame, Frame]()
-				{ 
-					FMemMark Mark(FMemStack::Get());
-					Simulate(SimResource, AnimFrame, Frame);
-				});
+				Simulate(SimResource, AnimFrame, Frame);
 			}
 			else
 			{
@@ -51,11 +47,6 @@ namespace UE::Chaos::FleshGenerator
 			}
 		}
 
-		for (TSharedPtr<FSimResource> SimResource : SimResources)
-		{
-			SimResource->Pipe->WaitUntilEmpty();
-		}
-	
 		RestoreAnimationSequence();
 	}
 
@@ -154,11 +145,7 @@ namespace UE::Chaos::FleshGenerator
 		check(SimResource.FleshComponent);
 		check(SimResource.SkeletalComponent);
 
-		// @todo(brice) : Do we need this?
-		// This could potentially be slow. 
 		SimResource.SkeletalComponent->RecreateRenderState_Concurrent();
-		SimResource.bNeedsSkin.store(true);
-		SimResource.SkinEvent->Wait();
 
 		TArray<FVector3f> Positions;
 		const UFleshDynamicAsset* DynamicCollection = SimResource.FleshComponent->GetDynamicCollection();
diff --git a/Engine/Plugins/Animation/MLDeformer/ChaosFleshGenerator/Source/ChaosFleshGenerator/Private/ChaosFleshGeneratorSimulation.h b/Engine/Plugins/Animation/MLDeformer/ChaosFleshGenerator/Source/ChaosFleshGenerator/Private/ChaosFleshGeneratorSimulation.h
--- a/Engine/Plugins/Animation/MLDeformer/ChaosFleshGenerator/Source/ChaosFleshGenerator/Private/ChaosFleshGeneratorSimulation.h
+++ b/Engine/Plugins/Animation/MLDeformer/ChaosFleshGenerator/Source/ChaosFleshGenerator/Private/ChaosFleshGeneratorSimulation.h
@@ -60,11 +60,6 @@ namespace UE::Chaos::FleshGenerator
 		TObjectPtr<USkeletalGeneratorComponent> SkeletalComponent = nullptr;
 		TObjectPtr<UDeformableSolverComponent> SolverComponent = nullptr;
 
-		TUniquePtr<::UE::Tasks::FPipe> Pipe;
-
-		FEvent* SkinEvent = nullptr;
-		std::atomic<bool> bNeedsSkin = false;
-
 		TArrayView<TArray<FVector3f>> SimulatedPositions;
 		std::atomic<int32>* NumSimulatedFrames = nullptr;
 
diff --git a/Engine/Plugins/Animation/MLDeformer/ChaosFleshGenerator/Source/ChaosFleshGenerator/Private/ChaosFleshGeneratorThreading.cpp b/Engine/Plugins/Animation/MLDeformer/ChaosFleshGenerator/Source/ChaosFleshGenerator/Private/ChaosFleshGeneratorThreading.cpp
--- a/Engine/Plugins/Animation/MLDeformer/ChaosFleshGenerator/Source/ChaosFleshGenerator/Private/ChaosFleshGeneratorThreading.cpp
+++ b/Engine/Plugins/Animation/MLDeformer/ChaosFleshGenerator/Source/ChaosFleshGenerator/Private/ChaosFleshGeneratorThreading.cpp
@@ -52,7 +52,6 @@ namespace UE::Chaos::FleshGenerator
 			DeformableSolver->SolverForces = Properties->SolverForces;
 			DeformableSolver->SolverDebugging = Properties->SolverDebugging;
 			DeformableSolver->Reset();
-			check(FleshComponent->GetPhysicsProxy());
 
 
 			constexpr int32 LODIndex = 0;
@@ -68,19 +67,10 @@ namespace UE::Chaos::FleshGenerator
 			SimResource.FleshComponent = TObjectPtr<UFleshGeneratorComponent>(FleshComponent);
 			SimResource.SkeletalComponent = TObjectPtr<USkeletalGeneratorComponent>(SkeletalMeshComponent);
 			SimResource.SolverComponent = TObjectPtr<UDeformableSolverComponent>(DeformableSolver);
-			SimResource.Pipe = MakeUnique<UE::Tasks::FPipe>(*FString::Printf(TEXT("SimPipe:%d"), Index));
-			SimResource.SkinEvent = FPlatformProcess::GetSynchEventFromPool();
-			SimResource.bNeedsSkin.store(false);
-	
+				
 			SimResource.SimulatedPositions = TArrayView<TArray<FVector3f>>(SimulatedPositions);
 			SimResource.NumSimulatedFrames = &NumSimulatedFrames;
 			SimResource.bCancelled = &bCancelled;
-	
-			if (FleshComponent == nullptr || SimResource.Pipe == nullptr)
-			{
-				UE_LOG(LogChaosFleshGenerator, Error, TEXT("Failed to allocate simulation resources"));
-				return false;
-			}
 		}
 
 		return true;
@@ -94,8 +84,6 @@ namespace UE::Chaos::FleshGenerator
 		}
 		for (TSharedPtr<FSimResource>& SimResource : SimResources)
 		{
-			FPlatformProcess::ReturnSynchEventToPool(SimResource->SkinEvent);
-			SimResource->Pipe.Reset();
 			SimResource->FleshComponent->UnregisterComponent();
 			SimResource->FleshComponent->DestroyComponent();
 		}
@@ -105,29 +93,6 @@ namespace UE::Chaos::FleshGenerator
 	
 	void FTaskResource::FlushRendering()
 	{
-		// Copy bNeedsSkin
-		TArray<bool> NeedsSkin;
-		NeedsSkin.SetNum(SimResources.Num());
-		bool bAnyNeedsSkin = false;
-		for (int32 Index = 0; Index < SimResources.Num(); ++Index)
-		{
-			const bool bNeedsSkin = SimResources[Index]->bNeedsSkin.load();
-			bAnyNeedsSkin |= bNeedsSkin;
-			NeedsSkin[Index] = bNeedsSkin;
-		}
-	
-		if (bAnyNeedsSkin)
-		{
-			FlushRenderingCommands();
-			for (int32 Index = 0; Index < SimResources.Num(); ++Index)
-			{
-				if (NeedsSkin[Index])
-				{
-					SimResources[Index]->bNeedsSkin.store(false);
-					SimResources[Index]->SkinEvent->Trigger();
-				}
-			}
-		}
 	}
 
 	void FTaskResource::Cancel()
