ChaosFlesh: Generate geometry cache in dataflow simulation

diff --git a/Engine/Plugins/Animation/MLDeformer/ChaosFleshGenerator/Source/ChaosFleshGenerator/ChaosFleshGenerator.build.cs b/Engine/Plugins/Animation/MLDeformer/ChaosFleshGenerator/Source/ChaosFleshGenerator/ChaosFleshGenerator.build.cs
--- a/Engine/Plugins/Animation/MLDeformer/ChaosFleshGenerator/Source/ChaosFleshGenerator/ChaosFleshGenerator.build.cs
+++ b/Engine/Plugins/Animation/MLDeformer/ChaosFleshGenerator/Source/ChaosFleshGenerator/ChaosFleshGenerator.build.cs
@@ -41,6 +41,7 @@ public class ChaosFleshGenerator : ModuleRules
 				"ChaosFleshEngine",
 				"ChaosFleshNodes",
 				"DataflowCore",
+				"DataflowEditor",
 				"DataflowEngine",
 				"DataflowSimulation",
 				"Engine",
diff --git a/Engine/Plugins/Animation/MLDeformer/ChaosFleshGenerator/Source/ChaosFleshGenerator/Private/ChaosFleshGenerator.cpp b/Engine/Plugins/Animation/MLDeformer/ChaosFleshGenerator/Source/ChaosFleshGenerator/Private/ChaosFleshGenerator.cpp
--- a/Engine/Plugins/Animation/MLDeformer/ChaosFleshGenerator/Source/ChaosFleshGenerator/Private/ChaosFleshGenerator.cpp
+++ b/Engine/Plugins/Animation/MLDeformer/ChaosFleshGenerator/Source/ChaosFleshGenerator/Private/ChaosFleshGenerator.cpp
@@ -7,6 +7,7 @@
 #include "Engine/SkeletalMesh.h"
 #include "FleshGeneratorProperties.h"
 #include "Misc/AsyncTaskNotification.h"
+#include "Dataflow/DataflowSimulationGeometryCache.h"
 
 DEFINE_LOG_CATEGORY(LogChaosFleshGenerator);
 
@@ -213,11 +214,8 @@ namespace UE::Chaos::FleshGenerator
 	
 		{
 			UE::Chaos::FleshGenerator::Private::FTimeScope TimeScope(TEXT("Saving"));
-	
-			using UE::Chaos::FleshGenerator::Private::SaveGeometryCache;
-			using UE::Chaos::FleshGenerator::Private::SavePackage;
-			SaveGeometryCache(*TaskResource->Cache, *Properties->SkeletalMeshAsset, TaskResource->ImportedVertexNumbers, TaskResource->SimulatedPositions);
-			SavePackage(*TaskResource->Cache);
+			DataflowSimulationGeometryCache::SaveGeometryCache(*TaskResource->Cache, *Properties->SkeletalMeshAsset, TaskResource->ImportedVertexNumbers, TaskResource->SimulatedPositions);
+			DataflowSimulationGeometryCache::SavePackage(*TaskResource->Cache);
 		}
 		if (bCancelled)
 		{
diff --git a/Engine/Plugins/Animation/MLDeformer/ChaosFleshGenerator/Source/ChaosFleshGenerator/Private/ChaosFleshGeneratorPrivate.cpp b/Engine/Plugins/Animation/MLDeformer/ChaosFleshGenerator/Source/ChaosFleshGenerator/Private/ChaosFleshGeneratorPrivate.cpp
--- a/Engine/Plugins/Animation/MLDeformer/ChaosFleshGenerator/Source/ChaosFleshGenerator/Private/ChaosFleshGeneratorPrivate.cpp
+++ b/Engine/Plugins/Animation/MLDeformer/ChaosFleshGenerator/Source/ChaosFleshGenerator/Private/ChaosFleshGeneratorPrivate.cpp
@@ -3,7 +3,7 @@
 #include "ChaosFleshGeneratorPrivate.h"
 
 #include "Chaos/Vector.h"
-#include "Dataflow/ChaosFleshGenerateSurfaceBindingsNode.h"
+#include "ChaosFlesh/FleshCollectionUtility.h"
 #include "FleshGeneratorComponent.h"
 #include "FileHelpers.h"
 #include "GeometryCollection/Facades/CollectionTetrahedralBindingsFacade.h"
@@ -16,8 +16,8 @@
 #include "Logging/LogMacros.h"
 #include "Rendering/SkeletalMeshRenderData.h"
 #include "Rendering/SkeletalMeshModel.h"
-
-
+#include "Components/SkeletalMeshComponent.h"
+#include "Engine/SkeletalMesh.h"
 DEFINE_LOG_CATEGORY(LogChaosFleshGeneratorPrivate);
 
 #define LOCTEXT_NAMESPACE "ChaosFleshGeneratorPrivate"
@@ -38,70 +38,6 @@ namespace UE::Chaos::FleshGenerator
 			UE_LOG(LogChaosFleshGeneratorPrivate, Log, TEXT("%s took %f secs"), *Name, Duration.GetTotalSeconds());
 		}
 
-
-		void BoundSurfacePositions(
-			const USkeletalMesh* SkeletalMesh,
-			const FFleshCollection* FleshCollection, 
-			const TManagedArray<FVector3f>* RestVertices,
-			const TManagedArray<FVector3f>* SimulatedVertices,
-			TArray<FVector3f>& Positions)
-		{
-			bool bError = false;
-
-			GeometryCollection::Facades::FTetrahedralBindings TetBindings(*FleshCollection);
-			FString MeshId = UE::TetrahedralBindingsEngineUtil::GetMeshId(SkeletalMesh, false);
-			FName MeshIdName(MeshId);
-
-			const int32 LODIndex = 0;
-			const int32 TetIndex = TetBindings.GetTetMeshIndex(MeshIdName, LODIndex);
-			if (TetIndex == INDEX_NONE)
-			{
-				UE_LOG(LogChaosFleshGeneratorPrivate, Error,
-					TEXT("CreateGeometryCache - No tet mesh index associated with mesh '%s' LOD: %d"),
-					*MeshId, LODIndex);
-				bError = true;
-			}
-			if (!TetBindings.ReadBindingsGroup(TetIndex, MeshIdName, LODIndex))
-			{
-				UE_LOG(LogChaosFleshGeneratorPrivate, Error,
-					TEXT("CreateGeometryCache - Failed to read bindings group associated with mesh '%s' LOD: %d"),
-					*MeshId, LODIndex);
-				bError = true;
-			}
-
-			TUniquePtr<GeometryCollection::Facades::FTetrahedralBindings::Evaluator> BindingsEvalPtr = TetBindings.InitEvaluator(RestVertices);
-			const GeometryCollection::Facades::FTetrahedralBindings::Evaluator& BindingsEval = *BindingsEvalPtr.Get();
-			if (!BindingsEval.IsValid())
-			{
-				UE_LOG(LogChaosFleshGeneratorPrivate, Error,
-					TEXT("CreateGeometryCache - Invalid flesh bindings for skeletal mesh asset [%s]"),
-					*SkeletalMesh->GetName());
-				bError = true;
-			}
-
-			if (!bError)
-			{
-				FVector3f Min = FLT_MAX*FVector3f::One();
-				FVector3f Max = FVector3f::Zero();
-
-				TArray<::Chaos::TVector<::Chaos::FRealSingle, 3>> CurrVertices;
-				CurrVertices.SetNum(SimulatedVertices->Num());
-				for (int32 Index = 0; Index < CurrVertices.Num(); ++Index)
-				{
-					CurrVertices[Index] = ::Chaos::TVector<::Chaos::FRealSingle, 3>((*SimulatedVertices)[Index]);
-					Min = FVector3f::Min(Min, (*SimulatedVertices)[Index]);
-					Max = FVector3f::Max(Max, (*SimulatedVertices)[Index]);
-				}
-
-				const int32 NumVertices = BindingsEval.NumVertices();
-				Positions.SetNum(NumVertices);
-				for (int32 Index = 0; Index < NumVertices; ++Index)
-				{
-					Positions[Index] = BindingsEval.GetEmbeddedPosition(Index, CurrVertices);
-				}
-			}
-		}
-
 		TArray<int32> ParseFrames(const FString& FramesString)
 		{
 			TArray<int32> Result;
@@ -165,65 +101,6 @@ namespace UE::Chaos::FleshGenerator
 			return Result;
 		}
 		
-		int32 GetNumVertices(const FSkeletalMeshLODRenderData& LODData)
-		{
-			int32 NumVertices = 0;
-			for(const FSkelMeshRenderSection& Section : LODData.RenderSections)
-			{
-				NumVertices += Section.NumVertices;
-			}
-			return NumVertices;
-		}
-		
-		TArrayView<TArray<FVector3f>> ShrinkToValidFrames(const TArrayView<TArray<FVector3f>>& Positions, int32 NumVertices)
-		{
-			int32 NumValidFrames = 0;
-			for (const TArray<FVector3f>& Frame : Positions)
-			{
-				if (Frame.Num() != NumVertices)
-				{
-					break;
-				}
-				++NumValidFrames;
-			}
-			return TArrayView<TArray<FVector3f>>(Positions.GetData(), NumValidFrames);
-		}
-		
-		void SaveGeometryCache(UGeometryCache& GeometryCache, const USkinnedAsset& Asset, TConstArrayView<uint32> ImportedVertexNumbers, TArrayView<TArray<FVector3f>> PositionsToMoveFrom)
-		{
-			const FSkeletalMeshRenderData* RenderData = Asset.GetResourceForRendering();
-			constexpr int32 LODIndex = 0;
-			if (!RenderData || !RenderData->LODRenderData.IsValidIndex(LODIndex))
-			{
-				return;
-			}
-			const FSkeletalMeshLODRenderData& LODData = RenderData->LODRenderData[LODIndex];
-			const int32 NumVertices = GetNumVertices(LODData);
-			PositionsToMoveFrom = ShrinkToValidFrames(PositionsToMoveFrom, NumVertices);
-		
-			using UE::GeometryCacheHelpers::FGeometryCacheConstantTopologyWriter;
-			using UE::GeometryCacheHelpers::AddTrackWriterFromSkinnedAsset;
-			UE::GeometryCacheHelpers::FGeometryCacheConstantTopologyWriter::FConfig Config; Config.FPS = 24;
-			using FTrackWriter = FGeometryCacheConstantTopologyWriter::FTrackWriter;
-			FGeometryCacheConstantTopologyWriter Writer(GeometryCache, Config);
-			const int32 Index = AddTrackWriterFromSkinnedAsset(Writer, Asset);
-			if (Index == INDEX_NONE)
-			{
-				return;
-			}
-			FTrackWriter& TrackWriter = Writer.GetTrackWriter(Index);
-			TrackWriter.ImportedVertexNumbers = ImportedVertexNumbers;
-			TrackWriter.WriteAndClose(PositionsToMoveFrom);
-		}
-		
-		void SavePackage(UObject& Object)
-		{
-			TArray<UPackage*> PackagesToSave = { Object.GetOutermost() };
-			constexpr bool bCheckDirty = false;
-			constexpr bool bPromptToSave = false;
-			FEditorFileUtils::PromptForCheckoutAndSave(PackagesToSave, bCheckDirty, bPromptToSave);
-		}
-		
 		TOptional<TArray<int32>> GetMeshImportVertexMap(const USkinnedAsset& SkinnedMeshAsset, const UFleshAsset& FleshAsset)
 		{
 			constexpr int32 LODIndex = 0;
@@ -254,8 +131,7 @@ namespace UE::Chaos::FleshGenerator
 			const TManagedArray<FVector3f>* RestVertices = FleshAsset.FindPositions();
 			if (SkeletalMeshAsset && FleshCollection && RestVertices)
 			{
-				using namespace UE::Chaos::FleshGenerator::Private;
-				BoundSurfacePositions(SkeletalMeshAsset, FleshCollection, RestVertices, RestVertices, Positions);
+				ChaosFlesh::BoundSurfacePositions(SkeletalMeshAsset, FleshCollection, RestVertices, RestVertices, Positions);
 
 				//@todo(Flesh Sections) : Add checks for multiple sections. 
 				int32 NumSections = MLDLOD.Sections.Num();
diff --git a/Engine/Plugins/Animation/MLDeformer/ChaosFleshGenerator/Source/ChaosFleshGenerator/Private/ChaosFleshGeneratorPrivate.h b/Engine/Plugins/Animation/MLDeformer/ChaosFleshGenerator/Source/ChaosFleshGenerator/Private/ChaosFleshGeneratorPrivate.h
--- a/Engine/Plugins/Animation/MLDeformer/ChaosFleshGenerator/Source/ChaosFleshGenerator/Private/ChaosFleshGeneratorPrivate.h
+++ b/Engine/Plugins/Animation/MLDeformer/ChaosFleshGenerator/Source/ChaosFleshGenerator/Private/ChaosFleshGeneratorPrivate.h
@@ -29,27 +29,12 @@ namespace UE::Chaos::FleshGenerator
 			~FTimeScope();
 		};
 
-		void BoundSurfacePositions(
-			const USkeletalMesh* SkeletalMesh,
-			const FFleshCollection* FleshCollection,
-			const TManagedArray<FVector3f>* RestVertices,
-			const TManagedArray<FVector3f>* SimulatedVertices,
-			TArray<FVector3f>& Positions);
-
 		TArray<int32> ParseFrames(const FString& FramesString);
 
 		TArray<int32> Range(int32 End);
 
 		TArray<uint32> Range(uint32 Start, uint32 End);
 
-		int32 GetNumVertices(const FSkeletalMeshLODRenderData& LODData);
-
-		TArrayView<TArray<FVector3f>> ShrinkToValidFrames(const TArrayView<TArray<FVector3f>>& Positions, int32 NumVertices);
-
-		void SaveGeometryCache(UGeometryCache& GeometryCache, const USkinnedAsset& Asset, TConstArrayView<uint32> ImportedVertexNumbers, TArrayView<TArray<FVector3f>> PositionsToMoveFrom);
-
-		void SavePackage(UObject& Object);
-
 		TOptional<TArray<int32>> GetMeshImportVertexMap(const USkinnedAsset& SkinnedMeshAsset, const UFleshAsset& FleshAsset);
 	};
 };
\ No newline at end of file
diff --git a/Engine/Plugins/Animation/MLDeformer/ChaosFleshGenerator/Source/ChaosFleshGenerator/Private/ChaosFleshGeneratorSimulation.cpp b/Engine/Plugins/Animation/MLDeformer/ChaosFleshGenerator/Source/ChaosFleshGenerator/Private/ChaosFleshGeneratorSimulation.cpp
--- a/Engine/Plugins/Animation/MLDeformer/ChaosFleshGenerator/Source/ChaosFleshGenerator/Private/ChaosFleshGeneratorSimulation.cpp
+++ b/Engine/Plugins/Animation/MLDeformer/ChaosFleshGenerator/Source/ChaosFleshGenerator/Private/ChaosFleshGeneratorSimulation.cpp
@@ -11,6 +11,7 @@
 #include "FleshGeneratorProperties.h"
 #include "Engine/SkeletalMesh.h"
 #include "Logging/LogMacros.h"
+#include "ChaosFlesh/FleshCollectionUtility.h"
 
 DEFINE_LOG_CATEGORY(LogChaosFleshGeneratorSimulation);
 
@@ -158,8 +159,7 @@ namespace UE::Chaos::FleshGenerator
 			const TManagedArray<FVector3f>* SimulatedVertices = DynamicCollection->FindPositions();
 			if (FleshCollection && RestVertices && SimulatedVertices)
 			{
-				using namespace UE::Chaos::FleshGenerator::Private;
-				BoundSurfacePositions(SkeletalMesh, FleshCollection, RestVertices, SimulatedVertices, Positions);
+				ChaosFlesh::BoundSurfacePositions(SkeletalMesh, FleshCollection, RestVertices, SimulatedVertices, Positions);
 			}
 		}
 		return Positions;
diff --git a/Engine/Plugins/Experimental/ChaosFlesh/Source/ChaosFlesh/ChaosFlesh.Build.cs b/Engine/Plugins/Experimental/ChaosFlesh/Source/ChaosFlesh/ChaosFlesh.Build.cs
--- a/Engine/Plugins/Experimental/ChaosFlesh/Source/ChaosFlesh/ChaosFlesh.Build.cs
+++ b/Engine/Plugins/Experimental/ChaosFlesh/Source/ChaosFlesh/ChaosFlesh.Build.cs
@@ -10,6 +10,7 @@ namespace UnrealBuildTool.Rules
 			PublicDependencyModuleNames.AddRange(
 				new string[]
 				{
+					"Engine",
 					"Core",
 					"CoreUObject",
 					"Chaos",
diff --git a/Engine/Plugins/Experimental/ChaosFlesh/Source/ChaosFlesh/Private/ChaosFlesh/FleshCollectionUtility.cpp b/Engine/Plugins/Experimental/ChaosFlesh/Source/ChaosFlesh/Private/ChaosFlesh/FleshCollectionUtility.cpp
--- a/Engine/Plugins/Experimental/ChaosFlesh/Source/ChaosFlesh/Private/ChaosFlesh/FleshCollectionUtility.cpp
+++ b/Engine/Plugins/Experimental/ChaosFlesh/Source/ChaosFlesh/Private/ChaosFlesh/FleshCollectionUtility.cpp
@@ -8,10 +8,11 @@
 
 #include "ChaosFlesh/ChaosFlesh.h"
 #include "ChaosFlesh/FleshCollection.h"
-
-
+#include "GeometryCollection/Facades/CollectionTetrahedralBindingsFacade.h"
 #include "GeometryCollection/TransformCollection.h"
 #include "Misc/Paths.h"
+#include "Engine/SkeletalMesh.h"
+#include "Engine/StaticMesh.h"
 
 #include <fstream>
 
@@ -45,40 +46,20 @@ namespace ChaosFlesh
 		}
 	}
 
-	int32 
-	GetMin(const FIntVector3& V) 
+	FIntVector3	GetOrdered(const FIntVector3& V)
 	{
-		return FMath::Min3(V[0], V[1], V[2]);
-	}
-	int32 
-	GetMid(const FIntVector3& V)
-	{
-		const int32 X = V[0]; const int32 Y = V[1]; const int32 Z = V[2];
-		const int32 XmY = X - Y;
-		const int32 YmZ = Y - Z;
-		const int32 XmZ = X - Z;
-		return (XmY * YmZ > -1 ? Y : XmY * XmZ < 1 ? X : Z);
-	}
-	int32 
-	GetMax(const FIntVector3& V)
-	{
-		return FMath::Max3(V[0], V[1], V[2]);
-	}
-	FIntVector3
-	GetOrdered(const FIntVector3& V)
-	{
-		return FIntVector3(GetMin(V), GetMid(V), GetMax(V));
+		TArray<int32, TInlineAllocator<3>> VA = { V[0], V[1], V[2] };
+		VA.Sort();
+		return FIntVector3(VA[0], VA[1], VA[2]);
 	}
-	FIntVector4
-	GetOrdered(const FIntVector4& V)
+	FIntVector4	GetOrdered(const FIntVector4& V)
 	{
-		TArray<int32> VA = { V[0], V[1], V[2], V[3] };
+		TArray<int32, TInlineAllocator<4>> VA = { V[0], V[1], V[2] };
 		VA.Sort();
 		return FIntVector4(VA[0], VA[1], VA[2], VA[3]);
 	}
 
-	void 
-	GetSurfaceElements(
+	void GetSurfaceElements(
 		const TArray<FIntVector4>& Tets,
 		TArray<FIntVector3>& SurfaceElements,
 		const bool KeepInteriorFaces,
@@ -174,4 +155,76 @@ namespace ChaosFlesh
 				UE_LOG(LogChaosFlesh, Display, TEXT("WARNING: Encountered %d non-manifold tetrahedral mesh faces."), nonManifold);
 		}
 	}
+
+	FString GetMeshId(const USkeletalMesh* SkeletalMesh, const bool bUseImportModel)
+	{
+		FPrimaryAssetId Id = SkeletalMesh->GetPrimaryAssetId();
+		FString MeshId = Id.IsValid() ? Id.ToString() : SkeletalMesh->GetName();
+		if (bUseImportModel)
+		{
+			MeshId.Append(TEXT("_ImportModel"));
+		}
+		return MeshId;
+	}
+
+	FString GetMeshId(const UStaticMesh* StaticMesh)
+	{
+		FPrimaryAssetId Id = StaticMesh->GetPrimaryAssetId();
+		FString MeshId = Id.IsValid() ? Id.ToString() : StaticMesh->GetName();
+		return MeshId;
+	}
+
+	void BoundSurfacePositions(
+		const USkeletalMesh* SkeletalMesh,
+		const FFleshCollection* FleshCollection,
+		const TManagedArray<FVector3f>* RestVertices,
+		const TManagedArray<FVector3f>* SimulatedVertices,
+		TArray<FVector3f>& Positions)
+	{
+		GeometryCollection::Facades::FTetrahedralBindings TetBindings(*FleshCollection);
+		FString MeshId = GetMeshId(SkeletalMesh, false);
+		FName MeshIdName(MeshId);
+
+		const int32 LODIndex = 0;
+		const int32 TetIndex = TetBindings.GetTetMeshIndex(MeshIdName, LODIndex);
+		if (TetIndex == INDEX_NONE)
+		{
+			UE_LOG(LogChaosFlesh, Error,
+				TEXT("CreateGeometryCache - No tet mesh index associated with mesh '%s' LOD: %d"),
+				*MeshId, LODIndex);
+			return;
+		}
+		if (!TetBindings.ReadBindingsGroup(TetIndex, MeshIdName, LODIndex))
+		{
+			UE_LOG(LogChaosFlesh, Error,
+				TEXT("CreateGeometryCache - Failed to read bindings group associated with mesh '%s' LOD: %d"),
+				*MeshId, LODIndex);
+			return;
+		}
+
+		TUniquePtr<GeometryCollection::Facades::FTetrahedralBindings::Evaluator> BindingsEvalPtr = TetBindings.InitEvaluator(RestVertices);
+		const GeometryCollection::Facades::FTetrahedralBindings::Evaluator& BindingsEval = *BindingsEvalPtr.Get();
+		if (!BindingsEval.IsValid())
+		{
+			UE_LOG(LogChaosFlesh, Error,
+				TEXT("CreateGeometryCache - Invalid flesh bindings for skeletal mesh asset [%s]"),
+				*SkeletalMesh->GetName());
+			return;
+		}
+
+		TArray<Chaos::FVec3f> CurrVertices;
+		CurrVertices.SetNum(SimulatedVertices->Num());
+		for (int32 Index = 0; Index < CurrVertices.Num(); ++Index)
+		{
+			CurrVertices[Index] = Chaos::FVec3f((*SimulatedVertices)[Index]);
+		}
+
+		const int32 NumVertices = BindingsEval.NumVertices();
+		Positions.SetNum(NumVertices);
+		for (int32 Index = 0; Index < NumVertices; ++Index)
+		{
+			Positions[Index] = BindingsEval.GetEmbeddedPosition(Index, CurrVertices);
+		}
+		
+	}
 }
diff --git a/Engine/Plugins/Experimental/ChaosFlesh/Source/ChaosFlesh/Public/ChaosFlesh/FleshCollectionUtility.h b/Engine/Plugins/Experimental/ChaosFlesh/Source/ChaosFlesh/Public/ChaosFlesh/FleshCollectionUtility.h
--- a/Engine/Plugins/Experimental/ChaosFlesh/Source/ChaosFlesh/Public/ChaosFlesh/FleshCollectionUtility.h
+++ b/Engine/Plugins/Experimental/ChaosFlesh/Source/ChaosFlesh/Public/ChaosFlesh/FleshCollectionUtility.h
@@ -7,8 +7,24 @@
 #include "Templates/UniquePtr.h"
 #include "Misc/AssertionMacros.h"
 
+class USkeletalMesh;
+class UFleshAsset;
+class UStaticMesh;
+
 namespace ChaosFlesh 
 {
 	void GetTetFaces(const FIntVector4& Tet, FIntVector3& Face1, FIntVector3& Face2, FIntVector3& Face3, FIntVector3& Face4, const bool invert = false);
+
 	void GetSurfaceElements(const TArray<FIntVector4>& Tets, TArray<FIntVector3>& SurfaceElements, const bool KeepInteriorFaces, const bool InvertFaces = false);
+	
+	CHAOSFLESH_API FString GetMeshId(const USkeletalMesh* SkeletalMesh, const bool bUseImportModel);
+
+	CHAOSFLESH_API FString GetMeshId(const UStaticMesh* StaticMesh);
+
+	CHAOSFLESH_API void BoundSurfacePositions(
+		const USkeletalMesh* SkeletalMesh,
+		const FFleshCollection* FleshCollection,
+		const TManagedArray<FVector3f>* RestVertices,
+		const TManagedArray<FVector3f>* SimulatedVertices,
+		TArray<FVector3f>& Positions);
 }
diff --git a/Engine/Plugins/Experimental/ChaosFlesh/Source/ChaosFleshEditor/Private/ChaosFlesh/Cmd/ChaosFleshCommands.cpp b/Engine/Plugins/Experimental/ChaosFlesh/Source/ChaosFleshEditor/Private/ChaosFlesh/Cmd/ChaosFleshCommands.cpp
--- a/Engine/Plugins/Experimental/ChaosFlesh/Source/ChaosFleshEditor/Private/ChaosFlesh/Cmd/ChaosFleshCommands.cpp
+++ b/Engine/Plugins/Experimental/ChaosFlesh/Source/ChaosFleshEditor/Private/ChaosFlesh/Cmd/ChaosFleshCommands.cpp
@@ -599,7 +599,7 @@ void FChaosFleshCommands::CreateGeometryCache(const TArray<FString>& Args, UWorl
 				// Extract bindings from tet mesh to skel mesh import geometry in the rest collection.
 				//
 
-				FString MeshId = UE::TetrahedralBindingsEngineUtil::GetMeshId(SkeletalMesh, false);
+				FString MeshId = ChaosFlesh::GetMeshId(SkeletalMesh, false);
 				FName MeshIdName(MeshId);
 				const int32 TetIndex = TetBindings.GetTetMeshIndex(MeshIdName, LODIndex);
 				if (TetIndex == INDEX_NONE)
diff --git a/Engine/Plugins/Experimental/ChaosFlesh/Source/ChaosFleshEngine/Private/ChaosCache/FleshComponentCacheAdapter.cpp b/Engine/Plugins/Experimental/ChaosFlesh/Source/ChaosFleshEngine/Private/ChaosCache/FleshComponentCacheAdapter.cpp
--- a/Engine/Plugins/Experimental/ChaosFlesh/Source/ChaosFleshEngine/Private/ChaosCache/FleshComponentCacheAdapter.cpp
+++ b/Engine/Plugins/Experimental/ChaosFlesh/Source/ChaosFleshEngine/Private/ChaosCache/FleshComponentCacheAdapter.cpp
@@ -443,6 +443,10 @@ namespace Chaos
 				{
 					UE::ChaosCachingUSD::OpenStage(PrimitiveStage.FilePath, PrimitiveStage.MonolithStage);
 				}
+				else
+				{
+					UE_LOG(LogChaosFleshCache, Warning, TEXT("Read Failure: USD File Path = %s"), *PrimitiveStage.FilePath);
+				}
 			}
 		}
 #endif // USE_USD_SDK && DO_USD_CACHING
diff --git a/Engine/Plugins/Experimental/ChaosFlesh/Source/ChaosFleshEngine/Private/ChaosFlesh/ChaosDeformableTetrahedralComponent.cpp b/Engine/Plugins/Experimental/ChaosFlesh/Source/ChaosFleshEngine/Private/ChaosFlesh/ChaosDeformableTetrahedralComponent.cpp
--- a/Engine/Plugins/Experimental/ChaosFlesh/Source/ChaosFleshEngine/Private/ChaosFlesh/ChaosDeformableTetrahedralComponent.cpp
+++ b/Engine/Plugins/Experimental/ChaosFlesh/Source/ChaosFleshEngine/Private/ChaosFlesh/ChaosDeformableTetrahedralComponent.cpp
@@ -12,6 +12,7 @@
 #include "ChaosFlesh/ChaosDeformableSolverComponent.h"
 #include "ChaosFlesh/ChaosDeformableTypes.h"
 #include "ChaosFlesh/ChaosFleshCollectionFacade.h"
+#include "ChaosFlesh/FleshCollectionUtility.h"
 #include "ChaosFlesh/FleshDynamicAsset.h"
 #include "Components/SkeletalMeshComponent.h"
 #include "Dataflow/DataflowEngineUtil.h"
@@ -23,6 +24,9 @@
 #include "ProceduralMeshComponent.h"
 #include "Rendering/SkeletalMeshRenderData.h"
 
+#if WITH_EDITOR
+#include "Rendering/SkeletalMeshModel.h"
+#endif
 
 #include UE_INLINE_GENERATED_CPP_BY_NAME(ChaosDeformableTetrahedralComponent)
 
@@ -624,6 +628,92 @@ TArray<FVector> UDeformableTetrahedralComponent::GetSkeletalMeshBindingPositions
 	return GetSkeletalMeshBindingPositionsInternal(InSkeletalMesh, nullptr);
 }
 
+TArray<FVector3f> UDeformableTetrahedralComponent::GetGeometryCachePositions(USkeletalMeshComponent* SkeletalComponent) const
+{
+	check(SkeletalComponent);
+	SkeletalComponent->RecreateRenderState_Concurrent();
+
+	TArray<FVector3f> Positions;
+	const USkeletalMesh* SkeletalMesh = SkeletalComponent->GetSkeletalMeshAsset();
+	if (RestCollection && DynamicCollection && SkeletalMesh)
+	{
+		const FFleshCollection* FleshCollection = RestCollection->GetCollection();
+		const TManagedArray<FVector3f>* RestVertices = RestCollection->FindPositions();
+		const TManagedArray<FVector3f>* SimulatedVertices = DynamicCollection->FindPositions();
+		if (FleshCollection && RestVertices && SimulatedVertices)
+		{
+			ChaosFlesh::BoundSurfacePositions(SkeletalMesh, FleshCollection, RestVertices, SimulatedVertices, Positions);
+		}
+	}
+	return Positions;
+}
+
+TOptional<TArray<int32>> UDeformableTetrahedralComponent::GetMeshImportVertexMap(const USkinnedAsset& SkinnedMeshAsset) const
+{
+	constexpr int32 LODIndex = 0;
+	const TOptional<TArray<int32>> None;
+#if WITH_EDITOR
+	const FSkeletalMeshModel* const MLDModel = SkinnedMeshAsset.GetImportedModel();
+	if (!MLDModel || !MLDModel->LODModels.IsValidIndex(LODIndex))
+	{
+		return None;
+	}
+	const FSkeletalMeshLODModel& MLDLOD = MLDModel->LODModels[LODIndex];
+	const TArray<int32>& Map = MLDLOD.MeshToImportVertexMap;
+	if (Map.IsEmpty())
+	{
+		UE_LOG(LogDeformableTetrahedralComponentInternal, Warning, TEXT("MeshToImportVertexMap is empty. MLDeformer Asset should be an imported SkeletalMesh (e.g. from fbx)."));
+		return None;
+	}
+
+	//
+	// @todo(flesh LOD) : Add support for managing vertex mappings between skeletal LOD.
+	//		The cloth/flesh asset will extract the LOD from the ManagedArrayCollection.
+
+	TArray<FVector3f> Positions;
+
+	const USkeletalMesh* SkeletalMeshAsset = Cast<USkeletalMesh>(&SkinnedMeshAsset);
+	const TManagedArray<FVector3f>* RestVertices = RestCollection->FindPositions();
+	if (SkeletalMeshAsset && RestCollection && RestVertices)
+	{
+		ChaosFlesh::BoundSurfacePositions(SkeletalMeshAsset, RestCollection->GetCollection(), RestVertices, RestVertices, Positions);
+
+		//@todo(Flesh Sections) : Add checks for multiple sections. 
+		int32 NumSections = MLDLOD.Sections.Num();
+		if (NumSections != 1)
+		{
+			UE_LOG(LogDeformableTetrahedralComponentInternal, Warning, TEXT("SkeletalMeshAsset should have only one section."));
+			return None;
+		}
+
+		for (int32 SectionIndex = 0; SectionIndex < NumSections; ++SectionIndex)
+		{
+			const FSkelMeshSection& MLDSection = MLDLOD.Sections[SectionIndex];
+			if (MLDSection.NumVertices != Positions.Num())
+			{
+				UE_LOG(LogDeformableTetrahedralComponentInternal, Warning, TEXT("SkeletalMeshAsset and FleshAsset have different number of vertices in section %d. Check if the assets have the same mesh."), SectionIndex);
+				return None;
+			}
+
+			for (int32 VertexIndex = 0; VertexIndex < MLDSection.NumVertices; ++VertexIndex)
+			{
+				const FVector3f& MLDPosition = MLDSection.SoftVertices[VertexIndex].Position;
+				const FVector3f& FleshPosition = Positions[VertexIndex];
+				if (!MLDPosition.Equals(FleshPosition, UE_KINDA_SMALL_NUMBER))
+				{
+					UE_LOG(LogDeformableTetrahedralComponentInternal, Warning, TEXT("SkeletalMeshAsset and FleshAsset have different vertex positions. Check if the assets have the same vertex order."));
+					return None;
+				}
+			}
+		}
+	}
+
+	return Map;
+#else
+	return None;
+#endif
+}
+
 TArray<FVector> UDeformableTetrahedralComponent::GetSkeletalMeshEmbeddedPositions(
 	const ChaosDeformableBindingOption Format,
 	const FTransform TargetDeformationSkeletonOffset,
diff --git a/Engine/Plugins/Experimental/ChaosFlesh/Source/ChaosFleshEngine/Public/ChaosFlesh/ChaosDeformableTetrahedralComponent.h b/Engine/Plugins/Experimental/ChaosFlesh/Source/ChaosFleshEngine/Public/ChaosFlesh/ChaosDeformableTetrahedralComponent.h
--- a/Engine/Plugins/Experimental/ChaosFlesh/Source/ChaosFleshEngine/Public/ChaosFlesh/ChaosDeformableTetrahedralComponent.h
+++ b/Engine/Plugins/Experimental/ChaosFlesh/Source/ChaosFleshEngine/Public/ChaosFlesh/ChaosDeformableTetrahedralComponent.h
@@ -10,6 +10,7 @@
 #include "ChaosFlesh/ChaosFleshDeformerBufferManager.h"
 #include "ChaosFlesh/SimulationAsset.h"
 #include "Components/MeshComponent.h"
+#include "Dataflow/Interfaces/DataflowInterfaceGeometryCachable.h"
 #include "UObject/ObjectMacros.h"
 #include "ProceduralMeshComponent.h"
 #include "ChaosDeformableTetrahedralComponent.generated.h"
@@ -18,6 +19,8 @@ class FFleshCollection;
 class ADeformableSolverActor;
 class UDeformableSolverComponent;
 class FChaosDeformableTetrahedralSceneProxy;
+class USkinnedAsset;
+class USkeletalMeshComponent;
 
 /**
 *  Options for binding positions query.
@@ -84,7 +87,7 @@ struct FBodyForcesGroup
 *	UDeformableTetrahedralComponent
 */
 UCLASS(meta = (BlueprintSpawnableComponent))
-class CHAOSFLESHENGINE_API UDeformableTetrahedralComponent : public UDeformablePhysicsComponent
+class CHAOSFLESHENGINE_API UDeformableTetrahedralComponent : public UDeformablePhysicsComponent, public IDataflowGeometryCachable
 {
 	GENERATED_UCLASS_BODY()
 
@@ -133,6 +136,12 @@ public:
 	/** @deprecated Use GetSkeletalMeshEmbeddedPositions() instead. */
 	UFUNCTION(BlueprintCallable, Category = "Physics", meta = (DeprecatedFunction, DeprecationMessage = "Use GetSkeletalMeshEmbeddedPositions() instead."))
 	TArray<FVector> GetSkeletalMeshBindingPositions(const USkeletalMesh* InSkeletalMesh) const;
+	
+	//~ Begin IDataflowGeometryCachable Interface
+	virtual TArray<FVector3f> GetGeometryCachePositions(USkeletalMeshComponent* SkeletalComponent) const override;
+
+	virtual TOptional<TArray<int32>> GetMeshImportVertexMap(const USkinnedAsset& SkinnedMeshAsset) const override;
+	//~ End IDataflowGeometryCachable Interface
 
 	/**
 	* Get the current positions of the transformation hierarchy from \c TargetDeformationSkeleton,
diff --git a/Engine/Plugins/Experimental/ChaosFlesh/Source/ChaosFleshNodes/Private/Dataflow/ChaosFleshGenerateSurfaceBindingsNode.cpp b/Engine/Plugins/Experimental/ChaosFlesh/Source/ChaosFleshNodes/Private/Dataflow/ChaosFleshGenerateSurfaceBindingsNode.cpp
--- a/Engine/Plugins/Experimental/ChaosFlesh/Source/ChaosFleshNodes/Private/Dataflow/ChaosFleshGenerateSurfaceBindingsNode.cpp
+++ b/Engine/Plugins/Experimental/ChaosFlesh/Source/ChaosFleshNodes/Private/Dataflow/ChaosFleshGenerateSurfaceBindingsNode.cpp
@@ -8,6 +8,7 @@
 #include "Chaos/BoundingVolumeHierarchy.h"
 #include "Chaos/Tetrahedron.h"
 #include "Chaos/TriangleMesh.h"
+#include "ChaosFlesh/FleshCollectionUtility.h"
 #include "ChaosFlesh/TetrahedralCollection.h"
 #include "Containers/Map.h"
 #include "Engine/StaticMesh.h"
@@ -20,31 +21,8 @@
 #include "Rendering/SkeletalMeshModel.h"
 #include "Rendering/SkeletalMeshRenderData.h"
 #include "UObject/PrimaryAssetId.h"
-
 DEFINE_LOG_CATEGORY(LogMeshBindings);
 
-namespace UE::TetrahedralBindingsEngineUtil
-{
-	FString GetMeshId(const USkeletalMesh* SkeletalMesh, const bool bUseImportModel)
-	{
-		FPrimaryAssetId Id = SkeletalMesh->GetPrimaryAssetId();
-		FString MeshId = Id.IsValid() ? Id.ToString() : SkeletalMesh->GetName();
-		if (bUseImportModel)
-		{
-			MeshId.Append(TEXT("_ImportModel"));
-		}
-		return MeshId;
-	}
-
-	FString GetMeshId(const UStaticMesh* StaticMesh)
-	{
-		FPrimaryAssetId Id = StaticMesh->GetPrimaryAssetId();
-		FString MeshId = Id.IsValid() ? Id.ToString() : StaticMesh->GetName();
-		return MeshId;
-	}
-}
-
-
 void
 BuildVertexToVertexAdjacencyBuffer(
 	const UE::Geometry::FDynamicMesh3 DynamicMesh,
@@ -195,7 +173,7 @@ FGenerateSurfaceBindings::Evaluate(Dataflow::FContext& Context, const FDataflowO
 			if (SkeletalMesh)
 			{
 				FPrimaryAssetId Id = SkeletalMesh->GetPrimaryAssetId();
-				MeshId = UE::TetrahedralBindingsEngineUtil::GetMeshId(SkeletalMesh, bUseSkeletalMeshImportModel);
+				MeshId = ChaosFlesh::GetMeshId(SkeletalMesh, bUseSkeletalMeshImportModel);
 
 				if (!bUseSkeletalMeshImportModel)
 				{
@@ -273,7 +251,7 @@ FGenerateSurfaceBindings::Evaluate(Dataflow::FContext& Context, const FDataflowO
 			}
 			else // StaticMesh
 			{
-				MeshId = UE::TetrahedralBindingsEngineUtil::GetMeshId(StaticMesh);
+				MeshId = ChaosFlesh::GetMeshId(StaticMesh);
 
 				const FStaticMeshRenderData* RenderData = StaticMesh->GetRenderData();
 				const int32 NumLOD = RenderData->LODResources.Num();
diff --git a/Engine/Plugins/Experimental/ChaosFlesh/Source/ChaosFleshNodes/Public/Dataflow/ChaosFleshGenerateSurfaceBindingsNode.h b/Engine/Plugins/Experimental/ChaosFlesh/Source/ChaosFleshNodes/Public/Dataflow/ChaosFleshGenerateSurfaceBindingsNode.h
--- a/Engine/Plugins/Experimental/ChaosFlesh/Source/ChaosFleshNodes/Public/Dataflow/ChaosFleshGenerateSurfaceBindingsNode.h
+++ b/Engine/Plugins/Experimental/ChaosFlesh/Source/ChaosFleshNodes/Public/Dataflow/ChaosFleshGenerateSurfaceBindingsNode.h
@@ -16,12 +16,6 @@ class USkeletalMesh;
 
 DECLARE_LOG_CATEGORY_EXTERN(LogMeshBindings, Verbose, All);
 
-namespace UE::TetrahedralBindingsEngineUtil
-{
-	CHAOSFLESHNODES_API FString GetMeshId(const USkeletalMesh* SkeletalMesh, const bool bUseImportModel = false);
-	CHAOSFLESHNODES_API FString GetMeshId(const UStaticMesh* StaticMesh);
-}
-
 /** Generate barycentric bindings (used by the FleshDeformer deformer graph) of a render surface to a tetrahedral mesh. */
 USTRUCT(meta = (DataflowFlesh))
 struct FGenerateSurfaceBindings : public FDataflowNode
diff --git a/Engine/Plugins/Experimental/Dataflow/Dataflow.uplugin b/Engine/Plugins/Experimental/Dataflow/Dataflow.uplugin
--- a/Engine/Plugins/Experimental/Dataflow/Dataflow.uplugin
+++ b/Engine/Plugins/Experimental/Dataflow/Dataflow.uplugin
@@ -58,6 +58,10 @@
 		{
 			"Name": "ChaosCaching",
 			"Enabled": true
+		},
+		{
+			"Name": "GeometryCache",
+			"Enabled": true
 		}
 	]
 }
diff --git a/Engine/Plugins/Experimental/Dataflow/Source/DataflowEditor/DataflowEditor.Build.cs b/Engine/Plugins/Experimental/Dataflow/Source/DataflowEditor/DataflowEditor.Build.cs
--- a/Engine/Plugins/Experimental/Dataflow/Source/DataflowEditor/DataflowEditor.Build.cs
+++ b/Engine/Plugins/Experimental/Dataflow/Source/DataflowEditor/DataflowEditor.Build.cs
@@ -36,6 +36,7 @@ namespace UnrealBuildTool.Rules
 					"EditorFramework",
 					"EditorInteractiveToolsFramework",
 					"EditorStyle",
+					"GeometryCache",
 					"GeometryCore",
 					"GeometryCollectionEngine",
 					"GeometryFramework",
diff --git a/Engine/Plugins/Experimental/Dataflow/Source/DataflowEditor/Private/Dataflow/DataflowSimulationGeometryCache.cpp b/Engine/Plugins/Experimental/Dataflow/Source/DataflowEditor/Private/Dataflow/DataflowSimulationGeometryCache.cpp
new file mode 100644
--- /dev/null
+++ b/Engine/Plugins/Experimental/Dataflow/Source/DataflowEditor/Private/Dataflow/DataflowSimulationGeometryCache.cpp
@@ -0,0 +1,87 @@
+// Copyright Epic Games, Inc. All Rights Reserved.
+
+#include "Dataflow/DataflowSimulationGeometryCache.h"
+
+#include "Chaos/Vector.h"
+#include "GeometryCache.h"
+#include "Logging/LogMacros.h"
+#include "Rendering/SkeletalMeshRenderData.h"
+#include "Rendering/SkeletalMeshModel.h"
+#include "Engine/SkeletalMesh.h"
+
+#if WITH_EDITOR 
+#include "FileHelpers.h"
+#include "GeometryCacheConstantTopologyWriter.h"
+#endif
+
+DEFINE_LOG_CATEGORY(LogDataflowSimulationGeometryCache);
+
+#define LOCTEXT_NAMESPACE "DataflowSimulationGeometryCache"
+	
+namespace DataflowSimulationGeometryCache
+{
+	int32 GetNumVertices(const FSkeletalMeshLODRenderData& LODData)
+	{
+		int32 NumVertices = 0;
+		for (const FSkelMeshRenderSection& Section : LODData.RenderSections)
+		{
+			NumVertices += Section.NumVertices;
+		}
+		return NumVertices;
+	}
+
+	TArrayView<TArray<FVector3f>> ShrinkToValidFrames(TArrayView<TArray<FVector3f>> Positions, int32 NumVertices)
+	{
+		int32 NumValidFrames = 0;
+		for (const TArray<FVector3f>& Frame : Positions)
+		{
+			if (Frame.Num() != NumVertices)
+			{
+				break;
+			}
+			++NumValidFrames;
+		}
+		return TArrayView<TArray<FVector3f>>(Positions.GetData(), NumValidFrames);
+	}
+
+	void SaveGeometryCache(UGeometryCache& GeometryCache, const USkinnedAsset& Asset, TConstArrayView<uint32> ImportedVertexNumbers, TArrayView<TArray<FVector3f>> PositionsToMoveFrom)
+	{
+#if WITH_EDITOR 
+		const FSkeletalMeshRenderData* RenderData = Asset.GetResourceForRendering();
+		constexpr int32 LODIndex = 0;
+		if (!RenderData || !RenderData->LODRenderData.IsValidIndex(LODIndex))
+		{
+			return;
+		}
+		const FSkeletalMeshLODRenderData& LODData = RenderData->LODRenderData[LODIndex];
+		const int32 NumVertices = GetNumVertices(LODData);
+		PositionsToMoveFrom = ShrinkToValidFrames(PositionsToMoveFrom, NumVertices);
+
+		using UE::GeometryCacheHelpers::FGeometryCacheConstantTopologyWriter;
+		using UE::GeometryCacheHelpers::AddTrackWriterFromSkinnedAsset;
+		FGeometryCacheConstantTopologyWriter::FConfig Config; Config.FPS = 24;
+		using FTrackWriter = FGeometryCacheConstantTopologyWriter::FTrackWriter;
+		FGeometryCacheConstantTopologyWriter Writer(GeometryCache, Config);
+		const int32 Index = AddTrackWriterFromSkinnedAsset(Writer, Asset);
+		if (Index == INDEX_NONE)
+		{
+			return;
+		}
+		FTrackWriter& TrackWriter = Writer.GetTrackWriter(Index);
+		TrackWriter.ImportedVertexNumbers = ImportedVertexNumbers;
+		TrackWriter.WriteAndClose(PositionsToMoveFrom);
+#endif
+	}
+
+	void SavePackage(UObject& Object)
+	{
+#if WITH_EDITOR 
+		TArray<UPackage*> PackagesToSave = { Object.GetOutermost() };
+		constexpr bool bCheckDirty = false;
+		constexpr bool bPromptToSave = false;
+		FEditorFileUtils::PromptForCheckoutAndSave(PackagesToSave, bCheckDirty, bPromptToSave);
+#endif
+	}
+};
+
+#undef LOCTEXT_NAMESPACE
diff --git a/Engine/Plugins/Experimental/Dataflow/Source/DataflowEditor/Private/Dataflow/DataflowSimulationScene.cpp b/Engine/Plugins/Experimental/Dataflow/Source/DataflowEditor/Private/Dataflow/DataflowSimulationScene.cpp
--- a/Engine/Plugins/Experimental/Dataflow/Source/DataflowEditor/Private/Dataflow/DataflowSimulationScene.cpp
+++ b/Engine/Plugins/Experimental/Dataflow/Source/DataflowEditor/Private/Dataflow/DataflowSimulationScene.cpp
@@ -12,6 +12,18 @@
 #include "AssetEditorModeManager.h"
 #include "Engine/Selection.h"
 
+#include "Dataflow/Interfaces/DataflowInterfaceGeometryCachable.h"
+#include "Dataflow/DataflowSimulationGeometryCache.h"
+#include "IContentBrowserSingleton.h"
+#include "ContentBrowserModule.h"
+#include "AssetRegistry/AssetRegistryModule.h"
+#include "Components/SkeletalMeshComponent.h"
+#include "Engine/SkeletalMesh.h"
+#include "GeometryCache.h"
+
+#if WITH_EDITOR
+#include "Misc/FileHelper.h"
+#endif
 #define LOCTEXT_NAMESPACE "FDataflowSimulationScene"
 
 //
@@ -261,6 +273,131 @@ void FDataflowSimulationScene::SceneDescriptionPropertyChanged(const FName& Prop
 	CreateSimulationScene();
 }
 
+void UDataflowSimulationSceneDescription::GenerateGeometryCache()
+{
+	SimulationScene->ResetSimulationScene();
+	SimulationScene->CreateSimulationScene();
+	const FVector2f& TimeRange = SimulationScene->GetTimeRange();
+	const int32 NumFrames = FMath::Floor((TimeRange[1] - TimeRange[0]) * CacheParams.FrameRate);
+	float Time = TimeRange[0];
+	float DeltaTime = (TimeRange[1] - TimeRange[0]) / NumFrames;
+	TObjectPtr<AActor> GetRootActor = SimulationScene->GetRootActor();
+	TObjectPtr<AActor> PreviewActor = SimulationScene->GetPreviewActor();
+	if (CacheAsset && GeometryCacheAsset && GetRootActor)
+	{
+		IDataflowGeometryCachable* GeometryCachable = nullptr; //interface for ChaosDeformableTetrahedralComponent
+
+		RenderPositions.SetNum(NumFrames);
+		TInlineComponentArray<UPrimitiveComponent*> PrimComponents;
+		PreviewActor->GetComponents(PrimComponents);
+		USkeletalMeshComponent* SkeletalComponent = nullptr;
+		for (UPrimitiveComponent* PrimComponent : PrimComponents)
+		{
+			
+			if (!GeometryCachable)
+			{
+				GeometryCachable = Cast<IDataflowGeometryCachable>(PrimComponent);
+			}
+			if (!SkeletalComponent)
+			{
+				SkeletalComponent = Cast<USkeletalMeshComponent>(PrimComponent);
+			}
+			if (GeometryCachable && SkeletalComponent)
+			{
+				break;
+			}
+		}
+		if (!GeometryCachable)
+		{
+			UE_LOG(LogChaosDataflow, Error, TEXT("No Flesh Component in the Preview Actor"));
+			return;
+		}
+		else if (!SkeletalComponent)
+		{
+			UE_LOG(LogChaosDataflow, Error, TEXT("No Skeletal Mesh Component in the Preview Actor"));
+			return;
+		}
+		TOptional<TArray<int32>> OptionalMap = GeometryCachable->GetMeshImportVertexMap(*SkeletalComponent->GetSkeletalMeshAsset());
+		if (!OptionalMap)
+		{
+			return;
+		}
+		const TArray<int32>& Map = OptionalMap.GetValue();
+		TArray<uint32> ImportedVertexNumbers = TArray<uint32>(reinterpret_cast<const uint32*>(Map.GetData()), Map.Num());
+		for (int32 Frame = 0; Frame < SimulationScene->GetNumFrames(); ++Frame)
+		{
+			Time += DeltaTime;
+			Cast<AChaosCacheManager>(GetRootActor)->SetStartTime(Time);
+			RenderPositions[Frame] = GeometryCachable->GetGeometryCachePositions(SkeletalComponent);
+		}
+		DataflowSimulationGeometryCache::SaveGeometryCache(*GeometryCacheAsset, *SkeletalComponent->GetSkeletalMeshAsset(), ImportedVertexNumbers, RenderPositions);
+		DataflowSimulationGeometryCache::SavePackage(*GeometryCacheAsset);
+	}
+}
+
+namespace Dataflow::Private
+{
+	template<class T>
+	T* CreateOrLoad(const FString& PackageName)
+	{
+		const FName AssetName(FPackageName::GetLongPackageAssetName(PackageName));
+		if (UPackage* const Package = CreatePackage(*PackageName))
+		{
+			LoadPackage(nullptr, *PackageName, LOAD_Quiet | LOAD_EditorOnly);
+			T* Asset = FindObject<T>(Package, *AssetName.ToString());
+			if (!Asset)
+			{
+				Asset = NewObject<T>(Package, *AssetName.ToString(), RF_Public | RF_Standalone | RF_Transactional);
+				Asset->MarkPackageDirty();
+				FAssetRegistryModule::AssetCreated(Asset);
+			}
+			return Asset;
+		}
+		return nullptr;
+	}
+
+	TObjectPtr<UGeometryCache> NewGeometryCacheDialog(const UObject* NamingAsset = nullptr)
+	{
+		FSaveAssetDialogConfig Config;
+		{
+			if (NamingAsset)
+			{
+				const FString PackageName = NamingAsset->GetOutermost()->GetName();
+				Config.DefaultPath = FPackageName::GetLongPackagePath(PackageName);
+				Config.DefaultAssetName = FString::Printf(TEXT("GeometryCache_%s"), *NamingAsset->GetName());
+			}
+			Config.AssetClassNames.Add(UGeometryCache::StaticClass()->GetClassPathName());
+			Config.ExistingAssetPolicy = ESaveAssetDialogExistingAssetPolicy::Disallow;
+			Config.DialogTitleOverride = LOCTEXT("ExportGeometryCacheDialogTitle", "Export Geometry Cache As");
+		}
+
+		FContentBrowserModule& ContentBrowserModule = FModuleManager::LoadModuleChecked<FContentBrowserModule>("ContentBrowser");
+
+#if WITH_EDITOR
+		FString NewPackageName;
+		FText OutError;
+		for (bool bFilenameValid = false; !bFilenameValid; bFilenameValid = FFileHelper::IsFilenameValidForSaving(NewPackageName, OutError))
+		{
+			const FString AssetPath = ContentBrowserModule.Get().CreateModalSaveAssetDialog(Config);
+			if (AssetPath.IsEmpty())
+			{
+				return nullptr;
+			}
+			NewPackageName = FPackageName::ObjectPathToPackageName(AssetPath);
+		}
+		return CreateOrLoad<UGeometryCache>(NewPackageName);
+#else
+		return nullptr;
+#endif
+	}
+};
+
+void UDataflowSimulationSceneDescription::NewGeometryCache()
+{
+	const UObject* const NamingAsset = CacheAsset? CacheAsset.Get() : nullptr;
+	GeometryCacheAsset = Dataflow::Private::NewGeometryCacheDialog(NamingAsset);
+}
+
 void UDataflowSimulationSceneDescription::SetSimulationScene(FDataflowSimulationScene* InSimulationScene)
 {
 	SimulationScene = InSimulationScene;
diff --git a/Engine/Plugins/Experimental/Dataflow/Source/DataflowEditor/Public/Dataflow/DataflowSimulationGeometryCache.h b/Engine/Plugins/Experimental/Dataflow/Source/DataflowEditor/Public/Dataflow/DataflowSimulationGeometryCache.h
new file mode 100644
--- /dev/null
+++ b/Engine/Plugins/Experimental/Dataflow/Source/DataflowEditor/Public/Dataflow/DataflowSimulationGeometryCache.h
@@ -0,0 +1,21 @@
+// Copyright Epic Games, Inc. All Rights Reserved.
+
+#pragma once
+
+#include "CoreTypes.h"
+#include "Containers/ContainersFwd.h"
+#include "GeometryCollection/ManagedArray.h"
+
+class USkinnedAsset;
+class UGeometryCache;
+
+DECLARE_LOG_CATEGORY_EXTERN(LogDataflowSimulationGeometryCache, Log, All);
+
+namespace DataflowSimulationGeometryCache
+{
+	/** Saves positions into the geometry cache */
+	DATAFLOWEDITOR_API void SaveGeometryCache(UGeometryCache& GeometryCache, const USkinnedAsset& Asset, TConstArrayView<uint32> ImportedVertexNumbers, TArrayView<TArray<FVector3f>> PositionsToMoveFrom);
+
+	/** Saves the geometry cache */
+	DATAFLOWEDITOR_API void SavePackage(UObject& Object);
+};
\ No newline at end of file
diff --git a/Engine/Plugins/Experimental/Dataflow/Source/DataflowEditor/Public/Dataflow/DataflowSimulationScene.h b/Engine/Plugins/Experimental/Dataflow/Source/DataflowEditor/Public/Dataflow/DataflowSimulationScene.h
--- a/Engine/Plugins/Experimental/Dataflow/Source/DataflowEditor/Public/Dataflow/DataflowSimulationScene.h
+++ b/Engine/Plugins/Experimental/Dataflow/Source/DataflowEditor/Public/Dataflow/DataflowSimulationScene.h
@@ -9,6 +9,11 @@
 #include "DataflowSimulationScene.generated.h"
 
 class UDataflowEditor;
+class UDeformableTetrahedralComponent;
+class USkeletalMeshComponent;
+class UGeometryCache;
+class UFleshDynamicAsset;
+class FFleshCollection;
 
 DECLARE_EVENT(UDataflowSimulationSceneDescription, FDataflowSimulationSceneDescriptionChanged)
 
@@ -40,14 +45,29 @@ public:
 	UPROPERTY(EditAnywhere, Category="Caching")
 	FDataflowPreviewCacheParams CacheParams;
 
+	/** Geometry cache asset used to extract skeletal mesh results from simulation */
+	UPROPERTY(EditAnywhere, Category = "Geometry Cache Conversion")
+	TObjectPtr<UGeometryCache> GeometryCacheAsset = nullptr;
+
+	/** Interpolates and saves geometry cache from Chaos cache */
+	UFUNCTION(CallInEditor, Category = "Geometry Cache Conversion")
+	void GenerateGeometryCache();
+
+	/** Creates a new geometry cache file */
+	UFUNCTION(CallInEditor, Category = "Geometry Cache Conversion")
+	void NewGeometryCache();
+
 private:
 
 	//~ UObject Interface
 	virtual void PostEditChangeProperty(struct FPropertyChangedEvent& PropertyChangedEvent) override;
 	virtual void PostTransacted(const FTransactionObjectEvent& TransactionEvent) override;
-
+	
 	/** Simulation scene linked to that descriptor */
 	class FDataflowSimulationScene* SimulationScene;
+
+	/** Render geometry positions from interpolation */
+	TArray<TArray<FVector3f>> RenderPositions;
 };
 
 /**
@@ -109,6 +129,10 @@ public:
 	/** Simulation time used to drive the cache loading */
 	float SimulationTime;
 	
+	/** Preview actor accessors */
+	TObjectPtr<AActor> GetPreviewActor() { return PreviewActor; }
+	const TObjectPtr<AActor> GetPreviewActor() const { return PreviewActor; }
+
 private:
 
 	/** Bind the scene selection to the components */
diff --git a/Engine/Source/Runtime/Experimental/Dataflow/Simulation/DataflowSimulation.Build.cs b/Engine/Source/Runtime/Experimental/Dataflow/Simulation/DataflowSimulation.Build.cs
--- a/Engine/Source/Runtime/Experimental/Dataflow/Simulation/DataflowSimulation.Build.cs
+++ b/Engine/Source/Runtime/Experimental/Dataflow/Simulation/DataflowSimulation.Build.cs
@@ -1,4 +1,4 @@
-﻿// Copyright Epic Games, Inc. All Rights Reserved.
+// Copyright Epic Games, Inc. All Rights Reserved.
 
 namespace UnrealBuildTool.Rules
 {
@@ -27,9 +27,13 @@ namespace UnrealBuildTool.Rules
 					"RHI",
 					"DataflowCore",
 					"DataflowEngine",
-					"Chaos"
+					"Chaos",
 				}
 			);
+			if (Target.bBuildEditor)
+			{
+				PrivateDependencyModuleNames.Add("UnrealEd");
+			}
 		}
 	}
 }
\ No newline at end of file
diff --git a/Engine/Source/Runtime/Experimental/Dataflow/Simulation/Public/Dataflow/Interfaces/DataflowInterfaceGeometryCachable.h b/Engine/Source/Runtime/Experimental/Dataflow/Simulation/Public/Dataflow/Interfaces/DataflowInterfaceGeometryCachable.h
new file mode 100644
--- /dev/null
+++ b/Engine/Source/Runtime/Experimental/Dataflow/Simulation/Public/Dataflow/Interfaces/DataflowInterfaceGeometryCachable.h
@@ -0,0 +1,32 @@
+// Copyright Epic Games, Inc. All Rights Reserved.
+
+#pragma once
+
+#include "UObject/ObjectMacros.h"
+#include "UObject/Interface.h"
+#include "Misc/Optional.h"
+#include "Containers/ArrayView.h"
+#include "Math/MathFwd.h"
+#include "DataflowInterfaceGeometryCachable.generated.h"
+
+class USkeletalMeshComponent;
+class USkinnedAsset;
+
+UINTERFACE(meta = (CannotImplementInterfaceInBlueprint), MinimalAPI)
+class UDataflowGeometryCachable : public UInterface
+{
+	GENERATED_BODY()
+};
+
+/** An asset that can produces a geometry cache */
+class IDataflowGeometryCachable
+{
+	GENERATED_BODY()
+
+public:
+	/** Computes skeletal mesh positions for geometry cache */
+	virtual TArray<FVector3f> GetGeometryCachePositions(USkeletalMeshComponent* SkeletalComponent) const = 0;
+
+	/** Get mapping from final mesh vertex index to raw import vertex index */
+	virtual TOptional<TArray<int32>> GetMeshImportVertexMap(const USkinnedAsset& SkinnedMeshAsset) const = 0;
+};
