Fix include for module manager

diff --git a/Engine/Plugins/Animation/MLDeformer/ChaosFleshGenerator/Source/ChaosFleshGenerator/Private/ChaosFleshGeneratorModule.cpp b/Engine/Plugins/Animation/MLDeformer/ChaosFleshGenerator/Source/ChaosFleshGenerator/Private/ChaosFleshGeneratorModule.cpp
--- a/Engine/Plugins/Animation/MLDeformer/ChaosFleshGenerator/Source/ChaosFleshGenerator/Private/ChaosFleshGeneratorModule.cpp
+++ b/Engine/Plugins/Animation/MLDeformer/ChaosFleshGenerator/Source/ChaosFleshGenerator/Private/ChaosFleshGeneratorModule.cpp
@@ -3,6 +3,7 @@
 #include "MLDeformerEditorToolkit.h"
 #include "PropertyEditorModule.h"
 #include "SFleshGeneratorWidget.h"
+#include "Modules/ModuleManager.h"
 
 namespace UE::Chaos::FleshGenerator
 {
