Fix for non-unity build errors

#FailedValidations Preflight

diff --git a/Engine/Plugins/Animation/MLDeformer/ChaosFleshGenerator/Source/ChaosFleshGenerator/Private/ChaosFleshGeneratorThreading.cpp b/Engine/Plugins/Animation/MLDeformer/ChaosFleshGenerator/Source/ChaosFleshGenerator/Private/ChaosFleshGeneratorThreading.cpp
--- a/Engine/Plugins/Animation/MLDeformer/ChaosFleshGenerator/Source/ChaosFleshGenerator/Private/ChaosFleshGeneratorThreading.cpp
+++ b/Engine/Plugins/Animation/MLDeformer/ChaosFleshGenerator/Source/ChaosFleshGenerator/Private/ChaosFleshGeneratorThreading.cpp
@@ -6,6 +6,7 @@
 #include "ChaosFlesh/ChaosDeformableSolverComponent.h"
 #include "ChaosFleshGenerator.h"
 #include "ChaosFleshGeneratorSimulation.h"
+#include "Engine/World.h"
 #include "FleshGeneratorComponent.h"
 #include "FleshGeneratorProperties.h"
 #include "GeometryCache.h"
diff --git a/Engine/Plugins/VirtualProduction/SceneState/Source/SceneStateMachineGraph/Private/Nodes/SceneStateMachineStateNode.cpp b/Engine/Plugins/VirtualProduction/SceneState/Source/SceneStateMachineGraph/Private/Nodes/SceneStateMachineStateNode.cpp
--- a/Engine/Plugins/VirtualProduction/SceneState/Source/SceneStateMachineGraph/Private/Nodes/SceneStateMachineStateNode.cpp
+++ b/Engine/Plugins/VirtualProduction/SceneState/Source/SceneStateMachineGraph/Private/Nodes/SceneStateMachineStateNode.cpp
@@ -2,6 +2,7 @@
 
 #include "Nodes/SceneStateMachineStateNode.h"
 #include "Kismet2/BlueprintEditorUtils.h"
+#include "Kismet2/Kismet2NameValidators.h"
 #include "Nodes/SceneStateMachineTaskNode.h"
 #include "SceneStateMachineGraph.h"
 #include "SceneStateMachineGraphSchema.h"
diff --git a/Engine/Plugins/VirtualProduction/SceneState/Source/SceneStateMachineGraph/Private/Nodes/SceneStateMachineTaskNode.cpp b/Engine/Plugins/VirtualProduction/SceneState/Source/SceneStateMachineGraph/Private/Nodes/SceneStateMachineTaskNode.cpp
--- a/Engine/Plugins/VirtualProduction/SceneState/Source/SceneStateMachineGraph/Private/Nodes/SceneStateMachineTaskNode.cpp
+++ b/Engine/Plugins/VirtualProduction/SceneState/Source/SceneStateMachineGraph/Private/Nodes/SceneStateMachineTaskNode.cpp
@@ -1,6 +1,7 @@
 // Copyright Epic Games, Inc. All Rights Reserved.
 
 #include "Nodes/SceneStateMachineTaskNode.h"
+#include "Kismet2/Kismet2NameValidators.h"
 #include "SceneStateBindingDelegates.h"
 #include "SceneStateBindingUtils.h"
 #include "SceneStateMachineGraphSchema.h"
